[RigMapper] second part of moving RigMapper Op into separate plugin

diff --git a/Engine/Plugins/Experimental/Animation/RigMapperOp/RigMapperOp.uplugin b/Engine/Plugins/Experimental/Animation/RigMapperOp/RigMapperOp.uplugin
new file mode 100644
--- /dev/null
+++ b/Engine/Plugins/Experimental/Animation/RigMapperOp/RigMapperOp.uplugin
@@ -0,0 +1,34 @@
+{
+	"FileVersion": 1,
+	"Version": 1,
+	"VersionName": "1.0",
+	"FriendlyName": "RigMapper Op",
+	"Description": "Experimental Retarget Op for re-mapping curves using RigMapper Definitions",
+	"Category": "Animation",
+	"CreatedBy": "Epic Games, Inc.",
+	"CreatedByURL": "https://epicgames.com",
+	"DocsURL": "",
+	"MarketplaceURL": "",
+	"SupportURL": "",
+	"CanContainContent": true,
+	"IsBetaVersion": false,
+	"IsExperimentalVersion": true,
+	"Installed": false,
+	"Modules": [
+		{
+			"Name": "RigMapperOp",
+			"Type": "Runtime",
+			"LoadingPhase": "Default"
+		}
+	],
+	"Plugins": [
+		{
+			"Name": "IKRig",
+			"Enabled": true
+		},
+		{
+			"Name": "RigMapper",
+			"Enabled": true
+		}
+	]
+}
\ No newline at end of file
diff --git a/Engine/Plugins/Experimental/Animation/RigMapperOp/Source/RigMapperOp/Private/RigMapperOp.cpp b/Engine/Plugins/Experimental/Animation/RigMapperOp/Source/RigMapperOp/Private/RigMapperOp.cpp
new file mode 100644
--- /dev/null
+++ b/Engine/Plugins/Experimental/Animation/RigMapperOp/Source/RigMapperOp/Private/RigMapperOp.cpp
@@ -0,0 +1,305 @@
+// Copyright Epic Games, Inc. All Rights Reserved.
+
+#include "RigMapperOp.h"
+
+#include "Animation/AnimCurveUtils.h"
+#include "Engine/SkeletalMesh.h"
+#include "Animation/AnimInstance.h"
+#include "Components/SkeletalMeshComponent.h"
+#include "Animation/AnimNodeBase.h"
+#include "RigMapperDefinition.h"
+#include "Retargeter/IKRetargetProcessor.h"
+
+#include UE_INLINE_GENERATED_CPP_BY_NAME(RigMapperOp)
+
+const UClass* FIKRetargetRigMapperOpSettings::GetControllerType() const
+{
+	return UIKRetargetRigMapperOpController::StaticClass();
+}
+
+void FIKRetargetRigMapperOpSettings::CopySettingsAtRuntime(const FIKRetargetOpSettingsBase* InSettingsToCopyFrom)
+{
+	// copies all properties
+	const TArray<FName> PropertiesToIgnore = {};
+	FIKRetargetOpBase::CopyStructProperties(
+		FIKRetargetRigMapperOpSettings::StaticStruct(),
+		InSettingsToCopyFrom,
+		this,
+		PropertiesToIgnore);
+}
+
+void FIKRetargetRigMapperOp::AnimGraphPreUpdateMainThread(
+	USkeletalMeshComponent& SourceMeshComponent,
+    USkeletalMeshComponent& TargetMeshComponent)
+{
+	if (!IsEnabled())
+	{
+		return;
+	}
+	
+	SourceCurves.Empty();
+	
+	// get the source curves out of the source anim instance
+	const UAnimInstance* SourceAnimInstance = SourceMeshComponent.GetAnimInstance();
+	if (!SourceAnimInstance)
+	{
+		return;
+	}
+	
+	// Potential optimization/tradeoff: If we stored the curve results on the mesh component in non-editor scenarios, this would be
+	// much faster (but take more memory). As it is, we need to translate the map stored on the anim instance.
+	const TMap<FName, float>& AnimCurveList = SourceAnimInstance->GetAnimationCurveList(EAnimCurveType::AttributeCurve);
+	UE::Anim::FCurveUtils::BuildUnsorted(SourceCurves, AnimCurveList);
+
+
+	// Need reinit if the definitions in use have changed.
+	bool bReInit = false;
+
+	// If definitions were loaded from the SKM asset user data, they take priority and are the ones we will check against
+	TArray<TObjectPtr<URigMapperDefinition>>* DefinitionsToCheck = &Settings.Definitions;
+	if (LoadedUserData)
+	{
+		DefinitionsToCheck = &LoadedUserData->Definitions;
+	}
+
+	check(DefinitionsToCheck);
+
+	// Otherwise, reinit if the definitions are not the same
+	if (DefinitionsToCheck->Num() != LoadedDefinitions.Num())
+	{
+		bReInit = true;
+	}
+	else
+	{
+		for (int32 DefinitionIndex = 0; DefinitionIndex < DefinitionsToCheck->Num(); DefinitionIndex++)
+		{
+			const bool bDefinitionEdittedAndNotValidated = (*DefinitionsToCheck)[DefinitionIndex] != nullptr && !(*DefinitionsToCheck)[DefinitionIndex]->WasDefinitionValidated();
+			const bool bDefinitionHasChanged = (*DefinitionsToCheck)[DefinitionIndex] != LoadedDefinitions[DefinitionIndex];
+			if (bDefinitionHasChanged || bDefinitionEdittedAndNotValidated)
+			{
+				bReInit = true;
+				break;
+			}
+		}
+	}
+	if (bReInit)
+	{
+		USkeletalMesh* TargetMesh = TargetMeshComponent.GetSkeletalMeshAsset();
+		bIsInitialized = InitializeRigMapping(TargetMesh);
+	}
+
+}
+
+void FIKRetargetRigMapperOp::AnimGraphEvaluateAnyThread(FPoseContext& Output)
+{
+	if (!IsEnabled())
+	{
+		return;
+	}
+
+	FBlendedCurve& OutputCurves = Output.Curve;
+
+	FBlendedCurve InputCurves;
+	if (GetTakeInputCurvesFromSourceAnimInstance())
+	{
+		InputCurves.CopyFrom(SourceCurves);
+		if (Settings.bCopyAllSourceCurves)
+		{
+			OutputCurves.CopyFrom(SourceCurves);
+		}
+	}
+	else
+	{
+		InputCurves.CopyFrom(OutputCurves);
+		// clear the outputs if needed
+		if (!Settings.bCopyAllSourceCurves)
+		{
+			OutputCurves.ForEachElement([&OutputCurves](const UE::Anim::FCurveElement& InCurveElement)
+				{
+					OutputCurves.Set(InCurveElement.Name, 0.0f);
+				});
+		}
+
+	}
+
+	EvaluateRigMapping(InputCurves, OutputCurves);
+
+}
+
+void FIKRetargetRigMapperOp::EvaluateRigMapping(const FBlendedCurve& InCurve, FBlendedCurve& OutCurve)
+{
+	if (!RigMapperProcessor.IsValid())
+	{
+		return;
+	}
+	
+	// Retrieve inputs
+	CachedInputValues.Reset(InputCurveMappings.Num());
+	CachedInputValues.AddDefaulted(InputCurveMappings.Num());
+	UE::Anim::FCurveUtils::BulkGet(InCurve, InputCurveMappings,
+		[&Inputs = CachedInputValues](const FRigMapperCurveMapping& InBulkElement, const float InValue)
+		{
+			Inputs[InBulkElement.CurveIndex] = InValue;
+		}
+	);
+
+	{
+		// Evaluate frame
+		RigMapperProcessor.EvaluateFrame(RigMapperProcessor.GetInputNames(), CachedInputValues, CachedOutputValues);
+	}
+
+	{
+		// Set all output curves from the given output values.
+		UE::Anim::FCurveUtils::BulkSet(OutCurve, OutputCurveMappings,
+			[&Inputs = CachedInputValues, &Outputs = CachedOutputValues](const FRigMapperOutputCurveMapping& InBulkElement)
+			{
+				TOptional<float> Value = Outputs[InBulkElement.CurveIndex];
+	
+				// If the value was not set by the rig mapper default to 0
+				return Value.Get(0.f);
+			}
+		);
+	}
+}
+
+
+bool FIKRetargetRigMapperOp::InitializeRigMapping(const USkeletalMesh* InTargetMesh)
+{
+	TArray<URigMapperDefinition*> DefinitionsToLoad;
+
+	// Retrieve the definitions to load and use (either the ones set in the op, or the ones overriden by the SKM asset user data)
+	if (InTargetMesh)
+	{
+		// Note that we cast away constness here as we need to get at the Asset User Data. This is not changed anywhere so is a reasonable pragmatic solution
+		LoadedUserData = Cast<URigMapperDefinitionUserData>(const_cast<USkeletalMesh*>(InTargetMesh)->GetAssetUserDataOfClass(URigMapperDefinitionUserData::StaticClass()));
+		if (LoadedUserData)
+		{
+			DefinitionsToLoad = LoadedUserData->Definitions;
+		}
+	}
+	if (DefinitionsToLoad.IsEmpty())
+	{
+		DefinitionsToLoad = Settings.Definitions;
+	}
+
+	RigMapperProcessor = FRigMapperProcessor(DefinitionsToLoad);
+	LoadedDefinitions = MoveTemp(DefinitionsToLoad);
+
+	if (!RigMapperProcessor.IsValid())
+	{
+		return false;
+	}
+
+	// Cache a map of curve indices to bulk get current curve values for the input
+	const TArray<FName>& InputNames = RigMapperProcessor.GetInputNames();
+	InputCurveMappings.Empty();
+	InputCurveMappings.Reserve(InputNames.Num());
+	for (int32 InputIndex = 0; InputIndex < InputNames.Num(); InputIndex++)
+	{
+		InputCurveMappings.Add(InputNames[InputIndex], InputIndex);
+	}
+
+	// Cache a map of curve indices to bulk set the new curve values for the output
+	const TArray<FName>& OutputNames = RigMapperProcessor.GetOutputNames();
+	OutputCurveMappings.Empty();
+	OutputCurveMappings.Reserve(OutputNames.Num());
+	for (int32 OutputIndex = 0; OutputIndex < OutputNames.Num(); OutputIndex++)
+	{
+		const FName& CurveName = OutputNames[OutputIndex];
+
+		int32 InputIndex = InputNames.Find(CurveName);
+		OutputCurveMappings.Add(OutputNames[OutputIndex], OutputIndex, InputIndex);
+	}
+
+	return true;
+}
+
+static FLinearColor MakeCurveColor(const FName& InCurveName)
+{
+	// Create a color based on the hash of the name
+	FRandomStream Stream(GetTypeHash(InCurveName));
+	const uint8 Hue = (uint8)(Stream.FRand() * 255.0f);
+	return FLinearColor::MakeFromHSV8(Hue, 196, 196);
+}
+
+void FIKRetargetRigMapperOp::ProcessAnimSequenceCurves(FIKRetargetOpBase::FCurveData InCurveMetaData, FIKRetargetOpBase::FFrameValues InCurveFrameValues,
+	FIKRetargetOpBase::FCurveData& OutCurveMetaData, FIKRetargetOpBase::FFrameValues& OutCurveFrameValues) const
+{
+	// We need to copy the RigMapperProcessor to keep this method const. Revisit this for speed at some point.
+	FRigMapperProcessor RigMapperProcessorTemp = RigMapperProcessor;
+	TSet<FName> OutputNames;
+
+	if (RigMapperProcessorTemp.IsValid())
+	{
+		RigMapperProcessorTemp.EvaluateFrames(InCurveMetaData.Names, InCurveFrameValues, OutCurveMetaData.Names, OutCurveFrameValues);
+		OutCurveMetaData.Flags.Init(4, OutCurveMetaData.Names.Num()); // a value of 4 means Editable curve
+		OutCurveMetaData.Colors.SetNum(OutCurveMetaData.Names.Num());
+		for (int32 Curve = 0; Curve < OutCurveMetaData.Names.Num(); ++Curve)
+		{
+			OutCurveMetaData.Colors[Curve] = MakeCurveColor(OutCurveMetaData.Names[Curve]);
+			OutputNames.Add(OutCurveMetaData.Names[Curve]);
+		}
+
+		// 'pass-through' any input curves which have not been modified, if bCopyAllSourceCurves is true
+		if (Settings.bCopyAllSourceCurves)
+		{
+			TArray<int32> InputCurveIndicesToAdd;
+			InputCurveIndicesToAdd.Reserve(InCurveMetaData.Names.Num());
+			for (int32 InputCurve = 0; InputCurve < InCurveMetaData.Names.Num(); ++InputCurve)
+			{
+				if (!OutputNames.Contains(InCurveMetaData.Names[InputCurve]))
+				{
+					InputCurveIndicesToAdd.Add(InputCurve);
+				}
+			}
+
+			// now add these curves
+			OutCurveMetaData.Names.Reserve(OutCurveMetaData.Names.Num() + InputCurveIndicesToAdd.Num());
+			OutCurveMetaData.Flags.Reserve(OutCurveMetaData.Flags.Num() + InputCurveIndicesToAdd.Num());
+			OutCurveMetaData.Colors.Reserve(OutCurveMetaData.Colors.Num() + InputCurveIndicesToAdd.Num());
+			int32 Frame = 0;
+			for (TArray<TOptional<float>>& OutputFrame : OutCurveFrameValues)
+			{
+				OutputFrame.Reserve(OutputFrame.Num() + InputCurveIndicesToAdd.Num());
+				for (int32 CurveIndex : InputCurveIndicesToAdd)
+				{
+					OutputFrame.Add(InCurveFrameValues[Frame][CurveIndex]);
+				}
+
+				Frame++;
+			}
+
+			for (int32 CurveIndex : InputCurveIndicesToAdd)
+			{
+				OutCurveMetaData.Names.Add(InCurveMetaData.Names[CurveIndex]);
+				OutCurveMetaData.Flags.Add(InCurveMetaData.Flags[CurveIndex]);
+				OutCurveMetaData.Colors.Add(InCurveMetaData.Colors[CurveIndex]);
+			}
+		}
+	}
+}
+
+bool FIKRetargetRigMapperOp::Initialize(
+	const FIKRetargetProcessor& InProcessor,
+	const FRetargetSkeleton& /*InSourceSkeleton*/,
+	const FTargetSkeleton& /*InTargetSkeleton*/,
+	const FIKRetargetOpBase* /*InParentOp*/,
+	FIKRigLogger& /*Log*/)
+{
+	const FRetargetSkeleton& TargetSkeleton = InProcessor.GetSkeleton(ERetargetSourceOrTarget::Target);
+	bIsInitialized = InitializeRigMapping(TargetSkeleton.SkeletalMesh);
+
+	return bIsInitialized;
+};
+
+
+
+FIKRetargetRigMapperOpSettings UIKRetargetRigMapperOpController::GetSettings()
+{
+	return *reinterpret_cast<FIKRetargetRigMapperOpSettings*>(OpSettingsToControl);
+}
+
+void UIKRetargetRigMapperOpController::SetSettings(FIKRetargetRigMapperOpSettings InSettings)
+{
+	OpSettingsToControl->CopySettingsAtRuntime(&InSettings);
+}
diff --git a/Engine/Plugins/Experimental/Animation/RigMapperOp/Source/RigMapperOp/Private/RigMapperOpModule.cpp b/Engine/Plugins/Experimental/Animation/RigMapperOp/Source/RigMapperOp/Private/RigMapperOpModule.cpp
new file mode 100644
--- /dev/null
+++ b/Engine/Plugins/Experimental/Animation/RigMapperOp/Source/RigMapperOp/Private/RigMapperOpModule.cpp
@@ -0,0 +1,17 @@
+// Copyright Epic Games, Inc. All Rights Reserved.
+
+#include "RigMapperOpModule.h"
+
+#define LOCTEXT_NAMESPACE "RigMapperOpModule"
+
+void FRigMapperOpModule::StartupModule()
+{
+}
+
+void FRigMapperOpModule::ShutdownModule()
+{
+}
+
+#undef LOCTEXT_NAMESPACE
+    
+IMPLEMENT_MODULE(FRigMapperOpModule, RigMapperOp)
\ No newline at end of file
diff --git a/Engine/Plugins/Experimental/Animation/RigMapperOp/Source/RigMapperOp/Public/RigMapperOp.h b/Engine/Plugins/Experimental/Animation/RigMapperOp/Source/RigMapperOp/Public/RigMapperOp.h
new file mode 100644
--- /dev/null
+++ b/Engine/Plugins/Experimental/Animation/RigMapperOp/Source/RigMapperOp/Public/RigMapperOp.h
@@ -0,0 +1,183 @@
+// Copyright Epic Games, Inc. All Rights Reserved.
+#pragma once
+
+#include "Retargeter/IKRetargetOps.h"
+#include "StructUtils/InstancedStruct.h"
+#include "Animation/AnimCurveTypes.h"
+#include "RigMapperProcessor.h"
+
+#include "RigMapperOp.generated.h"
+
+#define UE_API RIGMAPPEROP_API
+
+#define LOCTEXT_NAMESPACE "RigMapperOp"
+
+class URigMapperDefinitionUserData;
+class URigMapperDefinition;
+class USkeletalMesh;
+class USkeletalMeshComponent;
+
+USTRUCT(BlueprintType, meta = (DisplayName = "RigMapper Settings"))
+struct FIKRetargetRigMapperOpSettings : public FIKRetargetOpSettingsBase
+{
+	GENERATED_BODY()
+
+	// Whether to copy all curves over to the target animation instance
+	// NOTE: This setting also applies when exporting retargeted animations.
+	// True: all source curves are copied to the target animation instance/asset
+	// False: only remapped curves are copied on the target animation instance/asset
+	// In general, we should set this to true if the source and target rig are the same,
+	// and the RigMapper covers only a subset of the controls, but false otherwise.
+	UPROPERTY(EditAnywhere, BlueprintReadWrite, Category = "Copy Curves")
+	bool bCopyAllSourceCurves = false;
+
+	UPROPERTY(EditAnywhere, BlueprintReadWrite, Category = RigMapper)
+	TArray<TObjectPtr<URigMapperDefinition>> Definitions;
+
+	UE_API virtual const UClass* GetControllerType() const override;
+	
+	UE_API virtual void CopySettingsAtRuntime(const FIKRetargetOpSettingsBase* InSettingsToCopyFrom) override;
+};
+
+USTRUCT(BlueprintType, meta = (DisplayName = "RigMapper"))
+struct FIKRetargetRigMapperOp : public FIKRetargetOpBase
+{
+	GENERATED_BODY()
+	
+	
+	virtual bool Initialize(
+		const FIKRetargetProcessor& InProcessor,
+		const FRetargetSkeleton& InSourceSkeleton,
+		const FTargetSkeleton& InTargetSkeleton,
+		const FIKRetargetOpBase* InParentOp,
+		FIKRigLogger& Log) override;
+	
+	// NOTE: this op does not do anything in Run().
+	virtual void Run(
+		FIKRetargetProcessor& InProcessor,
+		const double InDeltaTime,
+		const TArray<FTransform>& InSourceGlobalPose,
+		TArray<FTransform>& OutTargetGlobalPose) override {};
+
+	virtual FIKRetargetOpSettingsBase* GetSettings() override
+	{
+		return &Settings;
+	};
+
+	virtual void SetSettings(const FIKRetargetOpSettingsBase* InSettings) override
+	{
+		Settings = *reinterpret_cast<const FIKRetargetRigMapperOpSettings*>(InSettings);
+	};
+	
+	virtual const UScriptStruct* GetSettingsType() const override
+	{
+		return FIKRetargetRigMapperOpSettings::StaticStruct();
+	}
+	
+	virtual const UScriptStruct* GetType() const override
+	{
+		return FIKRetargetRigMapperOp::StaticStruct();
+	}
+
+	/* RigMapperOp is a singleton because it can override definitions from the Target SkelMesh user data and therefore we should only do that once*/
+	virtual bool IsSingleton() const override { return true; };
+
+	virtual bool HasCurveProcessing() const override 
+	{ 
+		return true; 
+	}
+
+	virtual void ProcessAnimSequenceCurves(FIKRetargetOpBase::FCurveData InCurveMetaData, FIKRetargetOpBase::FFrameValues InCurveFrameValues,
+		FIKRetargetOpBase::FCurveData& OutCurveMetaData, FIKRetargetOpBase::FFrameValues& OutCurveFrameValues)const override;
+
+
+	UE_API virtual void AnimGraphPreUpdateMainThread(USkeletalMeshComponent& SourceMeshComponent, USkeletalMeshComponent& TargetMeshComponent) override;
+
+	UE_API virtual void AnimGraphEvaluateAnyThread(FPoseContext& Output) override;
+
+
+	UPROPERTY()
+	FIKRetargetRigMapperOpSettings Settings;
+	
+	// cached curves, copied on the game thread in PreUpdate()
+	FBlendedHeapCurve SourceCurves;
+
+
+private:
+	bool InitializeRigMapping(const USkeletalMesh* InTargetMesh);
+	void EvaluateRigMapping(const FBlendedCurve& InCurve, FBlendedCurve& OutCurve);
+
+	// The definitions that we have loaded. Cached to check against changes and reinit if need be
+	TArray<TObjectPtr<URigMapperDefinition>> LoadedDefinitions;
+
+	// The asset user data currently used to override definitions, if any was set on the skeletal mesh
+	TObjectPtr<URigMapperDefinitionUserData> LoadedUserData;
+
+	// The processor to evaluate the rig mapping
+	FRigMapperProcessor RigMapperProcessor;
+
+	// The cached input values passed to the rig mapper processor to avoid reallocations
+	FRigMapperProcessor::FPoseValues CachedInputValues;
+
+	// The cached output values passed to the rig mapper processor to avoid reallocations
+	FRigMapperProcessor::FPoseValues CachedOutputValues;
+
+	// Base curve mapping for bulk get/set of the linked pose curves
+	struct FRigMapperCurveMapping
+	{
+		FRigMapperCurveMapping() = default;
+
+		FRigMapperCurveMapping(FName InName, int32 InCurveIndex)
+			: Name(InName)
+			, CurveIndex(InCurveIndex)
+		{}
+
+		FName Name = NAME_None;
+		int32 CurveIndex = INDEX_NONE;
+	};
+
+	// Curve mapping for bulk set of the curves with a mapping to the matching input
+	struct FRigMapperOutputCurveMapping : FRigMapperCurveMapping
+	{
+		FRigMapperOutputCurveMapping() = default;
+
+		FRigMapperOutputCurveMapping(FName InName, int32 InOutputCurveIndex, int32 InInputCurveIndex)
+			: FRigMapperCurveMapping(InName, InOutputCurveIndex)
+			, InputCurveIndex(InInputCurveIndex)
+		{}
+
+		int32 InputCurveIndex = INDEX_NONE;
+	};
+
+	// Input curve mapping for bulk get of the curve values
+	using FInputCurveMappings = UE::Anim::TNamedValueArray<FDefaultAllocator, FRigMapperCurveMapping>;
+	FInputCurveMappings InputCurveMappings;
+
+	// Output curve mapping for the bulk get of the curve values
+	using FOutputCurveMappings = UE::Anim::TNamedValueArray<FDefaultAllocator, FRigMapperOutputCurveMapping>;
+	FOutputCurveMappings OutputCurveMappings;
+
+};
+
+/* The blueprint/python API for editing a RigMapper Op */
+UCLASS(MinimalAPI, BlueprintType)
+class UIKRetargetRigMapperOpController : public UIKRetargetOpControllerBase
+{
+	GENERATED_BODY()
+	
+public:
+	/* Get the current op settings as a struct.
+	 * @return FIKRetargetRigMapperOpSettings struct with the current settings used by the op. */
+	UFUNCTION(BlueprintCallable, Category = Settings)
+	UE_API FIKRetargetRigMapperOpSettings GetSettings();
+
+	/* Set the op settings. Input is a custom struct type for this op.
+	 * @param InSettings a FIKRetargetRigMapperOpSettings struct containing all the settings to apply to this op */
+	UFUNCTION(BlueprintCallable, Category = Settings)
+	UE_API void SetSettings(FIKRetargetRigMapperOpSettings InSettings);
+};
+
+
+#undef LOCTEXT_NAMESPACE
+
+#undef UE_API
diff --git a/Engine/Plugins/Experimental/Animation/RigMapperOp/Source/RigMapperOp/Public/RigMapperOpModule.h b/Engine/Plugins/Experimental/Animation/RigMapperOp/Source/RigMapperOp/Public/RigMapperOpModule.h
new file mode 100644
--- /dev/null
+++ b/Engine/Plugins/Experimental/Animation/RigMapperOp/Source/RigMapperOp/Public/RigMapperOpModule.h
@@ -0,0 +1,14 @@
+// Copyright Epic Games, Inc. All Rights Reserved.
+
+#pragma once
+
+#include "Modules/ModuleManager.h"
+
+class FRigMapperOpModule : public IModuleInterface
+{
+public:
+
+	/** IModuleInterface implementation */
+	virtual void StartupModule() override;
+	virtual void ShutdownModule() override;
+};
diff --git a/Engine/Plugins/Experimental/Animation/RigMapperOp/Source/RigMapperOp/RigMapperOp.Build.cs b/Engine/Plugins/Experimental/Animation/RigMapperOp/Source/RigMapperOp/RigMapperOp.Build.cs
new file mode 100644
--- /dev/null
+++ b/Engine/Plugins/Experimental/Animation/RigMapperOp/Source/RigMapperOp/RigMapperOp.Build.cs
@@ -0,0 +1,58 @@
+// Copyright Epic Games, Inc. All Rights Reserved.
+
+using UnrealBuildTool;
+using System.IO; // for Path
+
+namespace UnrealBuildTool.Rules
+{
+	public class RigMapperOp : ModuleRules
+	{
+		public RigMapperOp(ReadOnlyTargetRules Target) : base(Target)
+		{
+			PCHUsage = ModuleRules.PCHUsageMode.UseExplicitOrSharedPCHs;
+			CppCompileWarningSettings.UnsafeTypeCastWarningLevel = WarningLevel.Error;
+
+			PublicIncludePaths.AddRange(
+				new string[] {
+					// ... add public include paths required here ...
+				}
+				);
+
+
+			PrivateIncludePaths.AddRange(
+				new string[] {
+					// ... add other private include paths required here ...
+				}
+				);
+
+
+			PublicDependencyModuleNames.AddRange(
+				new string[]
+				{
+				}
+				);
+
+
+			PrivateDependencyModuleNames.AddRange(
+				new string[]
+				{
+				"Core",
+				"CoreUObject",
+				"Engine",
+				"AnimationCore",
+				"IKRig",
+				"RigMapper"
+				}
+				);
+
+
+			DynamicallyLoadedModuleNames.AddRange(
+				new string[]
+				{
+					// ... add any modules that your module loads dynamically here ...
+				}
+				);
+		}
+	}
+
+}
\ No newline at end of file
