IKRetargeter RigMapperOp refactoring

#Synced-CL 49271205

diff --git a/Engine/Plugins/Animation/IKRig/Source/IKRig/Private/Retargeter/IKRetargeter.cpp b/Engine/Plugins/Animation/IKRig/Source/IKRig/Private/Retargeter/IKRetargeter.cpp
--- a/Engine/Plugins/Animation/IKRig/Source/IKRig/Private/Retargeter/IKRetargeter.cpp
+++ b/Engine/Plugins/Animation/IKRig/Source/IKRig/Private/Retargeter/IKRetargeter.cpp
@@ -147,6 +147,7 @@ void UIKRetargeter::PostLoad()
 #if WITH_EDITOR
 	// NOTE: must run after PostLoadOpStack() to properly patch blend to source "Affect IK"
 	PostLoadSplitIKChainsOp();	
+	PostLoadSplitOps(); // Note: Generic ops splitting, i.e. when turning singleton op into mutliple ones
 #endif
 
 	CleanRetargetPoses();
@@ -745,6 +746,48 @@ void UIKRetargeter::PostLoadSplitIKChainsOp()
 		RemoveOp(Op);
 	}
 }
+
+void UIKRetargeter::PostLoadSplitOps()
+{
+	TArray<int32> OpsToRemove;
+	TArray<FInstancedStruct> OpsToAdd;
+
+	PRAGMA_DISABLE_DEPRECATION_WARNINGS
+
+	for (int32 OpIndex = 0; OpIndex < RetargetOps.Num(); OpIndex++)
+	{
+		FInstancedStruct& OpStruct = RetargetOps[OpIndex];
+		FIKRetargetOpBase* Op = OpStruct.GetMutablePtr<FIKRetargetOpBase>();
+		if (Op)
+		{
+			TArray<FInstancedStruct> NewOps;
+			// By default there will be no splitting
+			if (Op->PostLoadSplitOp(NewOps))
+			{
+				if (NewOps.Num() > 0)
+				{
+					// Mark index of the original op for removal
+					OpsToRemove.Add(OpIndex);
+
+					// Add the new ops
+					OpsToAdd.Append(MoveTemp(NewOps));
+
+					UE_LOG(LogTemp, Log, TEXT("IK Retargeter: Op %s split into %d new ops during PostLoad"), *Op->GetName().ToString(), NewOps.Num());
+				}
+			}
+		}
+	}
+	// Remove old ops in reverse order to maintain indices
+	for (int32 i = OpsToRemove.Num() - 1; i >= 0; i--)
+	{
+		RetargetOps.RemoveAt(OpsToRemove[i]);
+	}
+	// Add new ops
+	RetargetOps.Append(OpsToAdd);
+
+	PRAGMA_ENABLE_DEPRECATION_WARNINGS
+}
+
 #endif
 
 FName UIKRetargeter::GetCleanAndUniqueOpName(const FName& InOpName, const int32 InIndexOfOp)
diff --git a/Engine/Plugins/Animation/IKRig/Source/IKRig/Public/Retargeter/IKRetargetOps.h b/Engine/Plugins/Animation/IKRig/Source/IKRig/Public/Retargeter/IKRetargetOps.h
--- a/Engine/Plugins/Animation/IKRig/Source/IKRig/Public/Retargeter/IKRetargetOps.h
+++ b/Engine/Plugins/Animation/IKRig/Source/IKRig/Public/Retargeter/IKRetargetOps.h
@@ -167,7 +167,7 @@ struct FOpExecutionTimer
 	UE_API double GetRunningAverage() const;
 
 private:
-	/** running averrage of past n frames */
+	/** running average of past n frames */
 	double RunningAverageTime = 0.0f;
 	/** profiling data; time in seconds of last execution */
 	double TimeThisFrame = 0.0f;
@@ -234,9 +234,9 @@ public:
 	}
 
 	/** (optional) a function which returns true if the operation performs curve processing, false otherwise (default is false) */
-	virtual bool HasCurveProcessing() const 
+	virtual bool HasCurveProcessing() const
 	{
-		return false; 
+		return false;
 	}
 
 	/** a function which allows the user to manually indicate whether to take source curves from the source anim instance, or from the current pose context, since only the first enabled curve processing node must take inputs from the source anim instance*/
@@ -398,6 +398,17 @@ public:
 
 	/** a timer for the op profiler */
 	FOpExecutionTimer Timer;
+
+	/** Called during PostLoad to allow ops to be split into multiple ops during version upgrade.
+     * Return an array of new op instances that replace this op.
+     * If empty, no splitting occurred and the original op is preserved.
+     * @param OutNewOps - Out array which will hold new op instances (if splitting occurred)
+     * @return true if this op is being replaced by OutNewOps, false to keep original op
+	 */
+	virtual bool PostLoadSplitOp(TArray<FInstancedStruct>& OutNewOps)
+	{
+		return false; // No split by default.
+	}
 #endif
 
 	/** Get the name of this op (may be customized by user) */
diff --git a/Engine/Plugins/Animation/IKRig/Source/IKRig/Public/Retargeter/IKRetargeter.h b/Engine/Plugins/Animation/IKRig/Source/IKRig/Public/Retargeter/IKRetargeter.h
--- a/Engine/Plugins/Animation/IKRig/Source/IKRig/Public/Retargeter/IKRetargeter.h
+++ b/Engine/Plugins/Animation/IKRig/Source/IKRig/Public/Retargeter/IKRetargeter.h
@@ -313,6 +313,7 @@ private:
 	UE_API void PostLoadConvertEverythingToOps();
 	UE_API void PostLoadPutChainMappingInOps();
 	UE_API void PostLoadSplitIKChainsOp();
+	UE_API void PostLoadSplitOps();
 #endif
 	UE_API void CleanRetargetPoses();
 	UE_API void CleanOpStack();
diff --git a/Engine/Plugins/Animation/IKRig/Source/IKRig/Public/Retargeter/RetargetOps/CurveRemapOp.h b/Engine/Plugins/Animation/IKRig/Source/IKRig/Public/Retargeter/RetargetOps/CurveRemapOp.h
--- a/Engine/Plugins/Animation/IKRig/Source/IKRig/Public/Retargeter/RetargetOps/CurveRemapOp.h
+++ b/Engine/Plugins/Animation/IKRig/Source/IKRig/Public/Retargeter/RetargetOps/CurveRemapOp.h
@@ -57,7 +57,8 @@ USTRUCT(BlueprintType, meta = (DisplayName = "Remap Curves"))
 struct FIKRetargetCurveRemapOp : public FIKRetargetOpBase
 {
 	GENERATED_BODY()
-	
+
+public:
 	// NOTE: this op does not do anything in Initialize() or Run().
 	// It implements GetCurvesToRetarget() which the retargeting anim node calls and manages
 	// In the future we may remove this coupling and make the op itself do the work via callbacks
@@ -99,11 +100,16 @@ struct FIKRetargetCurveRemapOp : public FIKRetargetOpBase
 		return FIKRetargetCurveRemapOp::StaticStruct();
 	}
 
+	bool CanHaveChildOps() const override
+	{
+		return true;
+	}
+
 	virtual bool IsSingleton() const override { return true; };
 
-	virtual bool HasCurveProcessing() const override 
-	{ 
-		return true; 
+	virtual bool HasCurveProcessing() const override
+	{
+		return true;
 	}
 
 	virtual void ProcessAnimSequenceCurves(FIKRetargetOpBase::FCurveData InCurveMetaData, FIKRetargetOpBase::FFrameValues InCurveFrameValues,
@@ -112,12 +118,12 @@ struct FIKRetargetCurveRemapOp : public FIKRetargetOpBase
 	UE_API virtual void AnimGraphPreUpdateMainThread(USkeletalMeshComponent& SourceMeshComponent, USkeletalMeshComponent& TargetMeshComponent) override;
 
 	UE_API virtual void AnimGraphEvaluateAnyThread(FPoseContext& Output) override;
-
-	UPROPERTY()
-	FIKRetargetCurveRemapOpSettings Settings;
 	
 	// cached curves, copied on the game thread in PreUpdate()
 	FBlendedHeapCurve SourceCurves;
+
+	UPROPERTY()
+	FIKRetargetCurveRemapOpSettings Settings;
 };
 
 /* The blueprint/python API for editing a Curve Remap Op */
diff --git a/Engine/Plugins/Experimental/Animation/RigMapperOp/Source/RigMapperOp/Private/RigMapperOp.cpp b/Engine/Plugins/Experimental/Animation/RigMapperOp/Source/RigMapperOp/Private/RigMapperOp.cpp
--- a/Engine/Plugins/Experimental/Animation/RigMapperOp/Source/RigMapperOp/Private/RigMapperOp.cpp
+++ b/Engine/Plugins/Experimental/Animation/RigMapperOp/Source/RigMapperOp/Private/RigMapperOp.cpp
@@ -12,69 +12,89 @@
 
 #include UE_INLINE_GENERATED_CPP_BY_NAME(RigMapperOp)
 
-const UClass* FIKRetargetRigMapperOpSettings::GetControllerType() const
+static FLinearColor MakeCurveColor(const FName& InCurveName)
 {
-	return UIKRetargetRigMapperOpController::StaticClass();
+	// Create a color based on the hash of the name
+	FRandomStream Stream(GetTypeHash(InCurveName));
+	const uint8 Hue = (uint8)(Stream.FRand() * 255.0f);
+	return FLinearColor::MakeFromHSV8(Hue, 196, 196);
 }
 
-void FIKRetargetRigMapperOpSettings::CopySettingsAtRuntime(const FIKRetargetOpSettingsBase* InSettingsToCopyFrom)
+void FRigMapperOpHelper::ProcessAnimSequenceCurves(
+	FIKRetargetOpBase::FCurveData InCurveMetaData, 
+	FIKRetargetOpBase::FFrameValues InCurveFrameValues, 
+	FIKRetargetOpBase::FCurveData& OutCurveMetaData, 
+	FIKRetargetOpBase::FFrameValues& OutCurveFrameValues,
+	bool bCopyAllSourceCurves) const
 {
-	// copies all properties
-	const TArray<FName> PropertiesToIgnore = {};
-	FIKRetargetOpBase::CopyStructProperties(
-		FIKRetargetRigMapperOpSettings::StaticStruct(),
-		InSettingsToCopyFrom,
-		this,
-		PropertiesToIgnore);
+	// We need to copy the RigMapperProcessor to keep this method const. Revisit this for speed at some point.
+	FRigMapperProcessor RigMapperProcessorTemp = RigMapperProcessor;
+	TSet<FName> OutputNames;
+
+	if (RigMapperProcessorTemp.IsValid())
+	{
+		RigMapperProcessorTemp.EvaluateFrames(InCurveMetaData.Names, InCurveFrameValues, OutCurveMetaData.Names, OutCurveFrameValues);
+		OutCurveMetaData.Flags.Init(4, OutCurveMetaData.Names.Num()); // a value of 4 means Editable curve
+		OutCurveMetaData.Colors.SetNum(OutCurveMetaData.Names.Num());
+		for (int32 Curve = 0; Curve < OutCurveMetaData.Names.Num(); ++Curve)
+		{
+			OutCurveMetaData.Colors[Curve] = MakeCurveColor(OutCurveMetaData.Names[Curve]);
+			OutputNames.Add(OutCurveMetaData.Names[Curve]);
+		}
+
+		// 'pass-through' any input curves which have not been modified, if bCopyAllSourceCurves is true
+		if (bCopyAllSourceCurves)
+		{
+			TArray<int32> InputCurveIndicesToAdd;
+			InputCurveIndicesToAdd.Reserve(InCurveMetaData.Names.Num());
+			for (int32 InputCurve = 0; InputCurve < InCurveMetaData.Names.Num(); ++InputCurve)
+			{
+				if (!OutputNames.Contains(InCurveMetaData.Names[InputCurve]))
+				{
+					InputCurveIndicesToAdd.Add(InputCurve);
+				}
+			}
+
+			// now add these curves
+			OutCurveMetaData.Names.Reserve(OutCurveMetaData.Names.Num() + InputCurveIndicesToAdd.Num());
+			OutCurveMetaData.Flags.Reserve(OutCurveMetaData.Flags.Num() + InputCurveIndicesToAdd.Num());
+			OutCurveMetaData.Colors.Reserve(OutCurveMetaData.Colors.Num() + InputCurveIndicesToAdd.Num());
+			int32 Frame = 0;
+			for (TArray<TOptional<float>>& OutputFrame : OutCurveFrameValues)
+			{
+				OutputFrame.Reserve(OutputFrame.Num() + InputCurveIndicesToAdd.Num());
+				for (int32 CurveIndex : InputCurveIndicesToAdd)
+				{
+					OutputFrame.Add(InCurveFrameValues[Frame][CurveIndex]);
+				}
+
+				Frame++;
+			}
+
+			for (int32 CurveIndex : InputCurveIndicesToAdd)
+			{
+				OutCurveMetaData.Names.Add(InCurveMetaData.Names[CurveIndex]);
+				OutCurveMetaData.Flags.Add(InCurveMetaData.Flags[CurveIndex]);
+				OutCurveMetaData.Colors.Add(InCurveMetaData.Colors[CurveIndex]);
+			}
+		}
+	}
 }
 
-void FIKRetargetRigMapperOp::AnimGraphPreUpdateMainThread(
-	USkeletalMeshComponent& SourceMeshComponent,
-    USkeletalMeshComponent& TargetMeshComponent)
+bool FRigMapperOpHelper::CheckReInit(const TArray<URigMapperDefinition*>& InCurrentDefinitions)
 {
-	if (!IsEnabled())
-	{
-		return;
-	}
-	
-	SourceCurves.Empty();
-	
-	// get the source curves out of the source anim instance
-	const UAnimInstance* SourceAnimInstance = SourceMeshComponent.GetAnimInstance();
-	if (!SourceAnimInstance)
-	{
-		return;
-	}
-	
-	// Potential optimization/tradeoff: If we stored the curve results on the mesh component in non-editor scenarios, this would be
-	// much faster (but take more memory). As it is, we need to translate the map stored on the anim instance.
-	const TMap<FName, float>& AnimCurveList = SourceAnimInstance->GetAnimationCurveList(EAnimCurveType::AttributeCurve);
-	UE::Anim::FCurveUtils::BuildUnsorted(SourceCurves, AnimCurveList);
-
-
-	// Need reinit if the definitions in use have changed.
 	bool bReInit = false;
-
-	// If definitions were loaded from the SKM asset user data, they take priority and are the ones we will check against
-	TArray<TObjectPtr<URigMapperDefinition>>* DefinitionsToCheck = &Settings.Definitions;
-	if (LoadedUserData)
+	// Reinit if the definitions are not the same
+	if (InCurrentDefinitions.Num() != LoadedDefinitions.Num())
 	{
-		DefinitionsToCheck = &LoadedUserData->Definitions;
-	}
-
-	check(DefinitionsToCheck);
-
-	// Otherwise, reinit if the definitions are not the same
-	if (DefinitionsToCheck->Num() != LoadedDefinitions.Num())
-	{
 		bReInit = true;
 	}
 	else
 	{
-		for (int32 DefinitionIndex = 0; DefinitionIndex < DefinitionsToCheck->Num(); DefinitionIndex++)
+		for (int32 DefinitionIndex = 0; DefinitionIndex < InCurrentDefinitions.Num(); DefinitionIndex++)
 		{
-			const bool bDefinitionEdittedAndNotValidated = (*DefinitionsToCheck)[DefinitionIndex] != nullptr && !(*DefinitionsToCheck)[DefinitionIndex]->WasDefinitionValidated();
-			const bool bDefinitionHasChanged = (*DefinitionsToCheck)[DefinitionIndex] != LoadedDefinitions[DefinitionIndex];
+			const bool bDefinitionEdittedAndNotValidated = InCurrentDefinitions[DefinitionIndex] != nullptr && !InCurrentDefinitions[DefinitionIndex]->WasDefinitionValidated();
+			const bool bDefinitionHasChanged = InCurrentDefinitions[DefinitionIndex] != LoadedDefinitions[DefinitionIndex];
 			if (bDefinitionHasChanged || bDefinitionEdittedAndNotValidated)
 			{
 				bReInit = true;
@@ -82,57 +102,50 @@
 			}
 		}
 	}
-	if (bReInit)
+	return bReInit;
+}
+
+bool FRigMapperOpHelper::InitializeRigMapping(const TArray<URigMapperDefinition*>& InDefinitions)
+{
+	RigMapperProcessor = FRigMapperProcessor(InDefinitions);
+	LoadedDefinitions = InDefinitions;
+
+	if (!RigMapperProcessor.IsValid())
 	{
-		USkeletalMesh* TargetMesh = TargetMeshComponent.GetSkeletalMeshAsset();
-		bIsInitialized = InitializeRigMapping(TargetMesh);
+		return false;
 	}
 
+	// Cache a map of curve indices to bulk get current curve values for the input
+	const TArray<FName>& InputNames = RigMapperProcessor.GetInputNames();
+	InputCurveMappings.Empty();
+	InputCurveMappings.Reserve(InputNames.Num());
+	for (int32 InputIndex = 0; InputIndex < InputNames.Num(); InputIndex++)
+	{
+		InputCurveMappings.Add(InputNames[InputIndex], InputIndex);
+	}
+
+	// Cache a map of curve indices to bulk set the new curve values for the output
+	const TArray<FName>& OutputNames = RigMapperProcessor.GetOutputNames();
+	OutputCurveMappings.Empty();
+	OutputCurveMappings.Reserve(OutputNames.Num());
+	for (int32 OutputIndex = 0; OutputIndex < OutputNames.Num(); OutputIndex++)
+	{
+		const FName& CurveName = OutputNames[OutputIndex];
+
+		int32 InputIndex = InputNames.Find(CurveName);
+		OutputCurveMappings.Add(OutputNames[OutputIndex], OutputIndex, InputIndex);
+	}
+
+	return true;
 }
 
-void FIKRetargetRigMapperOp::AnimGraphEvaluateAnyThread(FPoseContext& Output)
-{
-	if (!IsEnabled())
-	{
-		return;
-	}
-
-	FBlendedCurve& OutputCurves = Output.Curve;
-
-	FBlendedCurve InputCurves;
-	if (GetTakeInputCurvesFromSourceAnimInstance())
-	{
-		InputCurves.CopyFrom(SourceCurves);
-		if (Settings.bCopyAllSourceCurves)
-		{
-			OutputCurves.CopyFrom(SourceCurves);
-		}
-	}
-	else
-	{
-		InputCurves.CopyFrom(OutputCurves);
-		// clear the outputs if needed
-		if (!Settings.bCopyAllSourceCurves)
-		{
-			OutputCurves.ForEachElement([&OutputCurves](const UE::Anim::FCurveElement& InCurveElement)
-				{
-					OutputCurves.Set(InCurveElement.Name, 0.0f);
-				});
-		}
-
-	}
-
-	EvaluateRigMapping(InputCurves, OutputCurves);
-
-}
-
-void FIKRetargetRigMapperOp::EvaluateRigMapping(const FBlendedCurve& InCurve, FBlendedCurve& OutCurve)
+void FRigMapperOpHelper::EvaluateRigMapping(const FBlendedCurve& InCurve, FBlendedCurve& OutCurve)
 {
 	if (!RigMapperProcessor.IsValid())
 	{
 		return;
 	}
-	
+
 	// Retrieve inputs
 	CachedInputValues.Reset(InputCurveMappings.Num());
 	CachedInputValues.AddDefaulted(InputCurveMappings.Num());
@@ -154,7 +167,7 @@
 			[&Inputs = CachedInputValues, &Outputs = CachedOutputValues](const FRigMapperOutputCurveMapping& InBulkElement)
 			{
 				TOptional<float> Value = Outputs[InBulkElement.CurveIndex];
-	
+
 				// If the value was not set by the rig mapper default to 0
 				return Value.Get(0.f);
 			}
@@ -162,144 +175,308 @@
 	}
 }
 
+URigMapperDefinitionUserData* FRigMapperOpHelper::GetUserDataFromMesh(const USkeletalMesh* InMesh)
+{
+	if (!InMesh) return nullptr;
 
-bool FIKRetargetRigMapperOp::InitializeRigMapping(const USkeletalMesh* InTargetMesh)
+	UAssetUserData* AssetUserData = const_cast<USkeletalMesh*>(InMesh)->
+		GetAssetUserDataOfClass(URigMapperDefinitionUserData::StaticClass());
+	return Cast<URigMapperDefinitionUserData>(AssetUserData);
+}
+
+bool FIKRetargetRigMapperOp::Initialize(const FIKRetargetProcessor& InProcessor, const FRetargetSkeleton& InSourceSkeleton, const FTargetSkeleton& InTargetSkeleton, const FIKRetargetOpBase* InParentOp, FIKRigLogger& Log)
+{
+	const FRetargetSkeleton& TargetSkeleton = InProcessor.GetSkeleton(ERetargetSourceOrTarget::Target);
+
+	TArray<URigMapperDefinition*> DefinitionsToLoad = GetDefinitionsToLoad(TargetSkeleton.SkeletalMesh);
+	bIsInitialized = Helper.InitializeRigMapping(DefinitionsToLoad);
+
+	return bIsInitialized;
+}
+
+void FIKRetargetRigMapperOp::ProcessAnimSequenceCurves(FIKRetargetOpBase::FCurveData InCurveMetaData, FIKRetargetOpBase::FFrameValues InCurveFrameValues, FIKRetargetOpBase::FCurveData& OutCurveMetaData, FIKRetargetOpBase::FFrameValues& OutCurveFrameValues) const
+{
+	Helper.ProcessAnimSequenceCurves(InCurveMetaData, InCurveFrameValues, OutCurveMetaData, OutCurveFrameValues, Settings.bCopyAllSourceCurves);
+}
+
+void FIKRetargetRigMapperOp::AnimGraphPreUpdateMainThread(USkeletalMeshComponent& SourceMeshComponent, USkeletalMeshComponent& TargetMeshComponent)
+{
+	if (!IsEnabled())
+	{
+		return;
+	}
+
+	SourceCurves.Empty();
+
+	// get the source curves out of the source anim instance
+	const UAnimInstance* SourceAnimInstance = SourceMeshComponent.GetAnimInstance();
+	if (!SourceAnimInstance)
+	{
+		return;
+	}
+
+	// Potential optimization/trade off: If we stored the curve results on the mesh component in non-editor scenarios, this would be
+	// much faster (but take more memory). As it is, we need to translate the map stored on the anim instance.
+	const TMap<FName, float>& AnimCurveList = SourceAnimInstance->GetAnimationCurveList(EAnimCurveType::AttributeCurve);
+	UE::Anim::FCurveUtils::BuildUnsorted(SourceCurves, AnimCurveList);
+
+	// Need reinit if the definitions in use have changed.
+	// If definitions were loaded from the SKM asset user data, they take priority and are the ones we will check against
+	USkeletalMesh* TargetMesh = TargetMeshComponent.GetSkeletalMeshAsset();
+	TArray<URigMapperDefinition*> DefinitionsToLoad = GetDefinitionsToLoad(TargetMesh);
+	if (Helper.CheckReInit(DefinitionsToLoad))
+	{
+		bIsInitialized = Helper.InitializeRigMapping(DefinitionsToLoad);
+	}
+}
+
+void FIKRetargetRigMapperOp::AnimGraphEvaluateAnyThread(FPoseContext& Output)
+{
+	if (!IsEnabled())
+	{
+		return;
+	}
+
+	FBlendedCurve& OutputCurves = Output.Curve;
+	FBlendedCurve InputCurves;
+
+	if (GetTakeInputCurvesFromSourceAnimInstance())
+	{
+		InputCurves.CopyFrom(SourceCurves);
+		if (Settings.bCopyAllSourceCurves)
+		{
+			OutputCurves.CopyFrom(SourceCurves);
+		}
+	}
+	else
+	{
+		InputCurves.CopyFrom(OutputCurves);
+		// clear the outputs if needed
+		if (!Settings.bCopyAllSourceCurves)
+		{
+			OutputCurves.ForEachElement([&OutputCurves](const UE::Anim::FCurveElement& InCurveElement)
+				{
+					OutputCurves.Set(InCurveElement.Name, 0.0f);
+				});
+		}
+	}
+
+	Helper.EvaluateRigMapping(InputCurves, OutputCurves);
+}
+
+bool FIKRetargetRigMapperOp::PostLoadSplitOp(TArray<FInstancedStruct>& OutNewOps)
+{
+	PRAGMA_DISABLE_DEPRECATION_WARNINGS
+	// Only split if there are multiple definitions in the old RigMapperOp
+	if (Settings.Definitions.Num() <= 1)
+	{
+		// Migrate single definition and empty array
+		if (Settings.Definitions.Num() == 1)
+		{
+			Settings.Definition = Settings.Definitions[0];
+			Settings.Definitions.Empty();
+		}
+		return false; // Don't split, just migrate in place
+	}
+
+	// In case of multiple definitions create one op per definition
+	const FName OriginalName = GetName();
+	const bool bOriginalEnabled = IsEnabled();
+	const FName ParentName = GetParentOpName();
+
+	for (int32 DefinitionIndex = 0; DefinitionIndex < Settings.Definitions.Num(); DefinitionIndex++)
+	{
+		// Create new op instance
+		FInstancedStruct NewOpStruct;
+		NewOpStruct.InitializeAs(FIKRetargetRigMapperOp::StaticStruct());
+		FIKRetargetRigMapperOp* NewOp = NewOpStruct.GetMutablePtr<FIKRetargetRigMapperOp>();
+
+		if (NewOp)
+		{
+			// Copy settings
+			NewOp->Settings.bCopyAllSourceCurves = Settings.bCopyAllSourceCurves;
+			// Set single definition (new property) instead of deprecated array
+			NewOp->Settings.Definition = Settings.Definitions[DefinitionIndex];
+
+			// Copy enabled state
+			NewOp->SetEnabled(bOriginalEnabled);
+
+			// Set parent op
+			if (!ParentName.IsNone())
+			{
+				NewOp->SetParentOpName(ParentName);
+			}
+
+			OutNewOps.Add(MoveTemp(NewOpStruct));
+		}
+	}
+
+	return true; // Replace original op with new ops
+	PRAGMA_ENABLE_DEPRECATION_WARNINGS
+}
+
+TArray<URigMapperDefinition*> FIKRetargetRigMapperOp::GetDefinitionsToLoad(const USkeletalMesh* InTargetMesh)
 {
 	TArray<URigMapperDefinition*> DefinitionsToLoad;
 
-	// Retrieve the definitions to load and use (either the ones set in the op, or the ones overriden by the SKM asset user data)
-	if (InTargetMesh)
+	// If bOverrideFromUserDataDefinitions is true, get definitions from UserDataAsset
+	if (Settings.bOverrideFromUserDataDefinitions)
 	{
-		// Note that we cast away constness here as we need to get at the Asset User Data. This is not changed anywhere so is a reasonable pragmatic solution
-		LoadedUserData = Cast<URigMapperDefinitionUserData>(const_cast<USkeletalMesh*>(InTargetMesh)->GetAssetUserDataOfClass(URigMapperDefinitionUserData::StaticClass()));
-		if (LoadedUserData)
+		URigMapperDefinitionUserData* UserData = FRigMapperOpHelper::GetUserDataFromMesh(InTargetMesh);
+		if (UserData)
 		{
-			DefinitionsToLoad = LoadedUserData->Definitions;
+			DefinitionsToLoad = UserData->Definitions;
 		}
 	}
-	if (DefinitionsToLoad.IsEmpty())
+	// Otherwise use single definition from the Op
+	else if (Settings.Definition)
 	{
-		DefinitionsToLoad = Settings.Definitions;
+		DefinitionsToLoad.Add(Settings.Definition);
+	}
+	return DefinitionsToLoad;
+}
+
+const UClass* FIKRetargetRigMapperOpSettings::GetControllerType() const
+{
+	return UIKRetargetRigMapperOpController::StaticClass();
+}
+
+void FIKRetargetRigMapperOpSettings::CopySettingsAtRuntime(const FIKRetargetOpSettingsBase* InSettingsToCopyFrom)
+{
+	// copies all properties
+	const TArray<FName> PropertiesToIgnore = {};
+	FIKRetargetOpBase::CopyStructProperties(
+		FIKRetargetRigMapperOpSettings::StaticStruct(),
+		InSettingsToCopyFrom,
+		this,
+		PropertiesToIgnore);
+}
+
+bool FIKRetargetRigMapperUserDataOp::Initialize(const FIKRetargetProcessor& InProcessor, const FRetargetSkeleton& InSourceSkeleton, const FTargetSkeleton& InTargetSkeleton, const FIKRetargetOpBase* InParentOp, FIKRigLogger& Log)
+{
+	const FRetargetSkeleton& TargetSkeleton = InProcessor.GetSkeleton(ERetargetSourceOrTarget::Target);
+
+	TArray<URigMapperDefinition*> DefinitionsToLoad = GetDefinitionsToLoad(TargetSkeleton.SkeletalMesh);
+	bIsInitialized = Helper.InitializeRigMapping(DefinitionsToLoad);
+
+	return bIsInitialized;
+}
+
+void FIKRetargetRigMapperUserDataOp::ProcessAnimSequenceCurves(FIKRetargetOpBase::FCurveData InCurveMetaData, FIKRetargetOpBase::FFrameValues InCurveFrameValues, FIKRetargetOpBase::FCurveData& OutCurveMetaData, FIKRetargetOpBase::FFrameValues& OutCurveFrameValues) const
+{
+	Helper.ProcessAnimSequenceCurves(InCurveMetaData, InCurveFrameValues, OutCurveMetaData, OutCurveFrameValues, Settings.bCopyAllSourceCurves);
+}
+
+void FIKRetargetRigMapperUserDataOp::AnimGraphPreUpdateMainThread(USkeletalMeshComponent& SourceMeshComponent, USkeletalMeshComponent& TargetMeshComponent)
+{
+	if (!IsEnabled())
+	{
+		return;
 	}
 
-	RigMapperProcessor = FRigMapperProcessor(DefinitionsToLoad);
-	LoadedDefinitions = MoveTemp(DefinitionsToLoad);
+	SourceCurves.Empty();
 
-	if (!RigMapperProcessor.IsValid())
+	// get the source curves out of the source anim instance
+	const UAnimInstance* SourceAnimInstance = SourceMeshComponent.GetAnimInstance();
+	if (!SourceAnimInstance)
 	{
-		return false;
+		return;
 	}
 
-	// Cache a map of curve indices to bulk get current curve values for the input
-	const TArray<FName>& InputNames = RigMapperProcessor.GetInputNames();
-	InputCurveMappings.Empty();
-	InputCurveMappings.Reserve(InputNames.Num());
-	for (int32 InputIndex = 0; InputIndex < InputNames.Num(); InputIndex++)
+	// Potential optimization/trade off: If we stored the curve results on the mesh component in non-editor scenarios, this would be
+	// much faster (but take more memory). As it is, we need to translate the map stored on the anim instance.
+	const TMap<FName, float>& AnimCurveList = SourceAnimInstance->GetAnimationCurveList(EAnimCurveType::AttributeCurve);
+	UE::Anim::FCurveUtils::BuildUnsorted(SourceCurves, AnimCurveList);
+
+	// Need reinit if the definitions in use have changed.
+	// If definitions were loaded from the SKM asset user data, they take priority and are the ones we will check against
+	USkeletalMesh* TargetMesh = TargetMeshComponent.GetSkeletalMeshAsset();
+	TArray<URigMapperDefinition*> DefinitionsToLoad = GetDefinitionsToLoad(TargetMesh);
+	if (Helper.CheckReInit(DefinitionsToLoad))
 	{
-		InputCurveMappings.Add(InputNames[InputIndex], InputIndex);
-	}
-
-	// Cache a map of curve indices to bulk set the new curve values for the output
-	const TArray<FName>& OutputNames = RigMapperProcessor.GetOutputNames();
-	OutputCurveMappings.Empty();
-	OutputCurveMappings.Reserve(OutputNames.Num());
-	for (int32 OutputIndex = 0; OutputIndex < OutputNames.Num(); OutputIndex++)
-	{
-		const FName& CurveName = OutputNames[OutputIndex];
-
-		int32 InputIndex = InputNames.Find(CurveName);
-		OutputCurveMappings.Add(OutputNames[OutputIndex], OutputIndex, InputIndex);
-	}
-
-	return true;
-}
-
-static FLinearColor MakeCurveColor(const FName& InCurveName)
-{
-	// Create a color based on the hash of the name
-	FRandomStream Stream(GetTypeHash(InCurveName));
-	const uint8 Hue = (uint8)(Stream.FRand() * 255.0f);
-	return FLinearColor::MakeFromHSV8(Hue, 196, 196);
-}
-
-void FIKRetargetRigMapperOp::ProcessAnimSequenceCurves(FIKRetargetOpBase::FCurveData InCurveMetaData, FIKRetargetOpBase::FFrameValues InCurveFrameValues,
-	FIKRetargetOpBase::FCurveData& OutCurveMetaData, FIKRetargetOpBase::FFrameValues& OutCurveFrameValues) const
-{
-	// We need to copy the RigMapperProcessor to keep this method const. Revisit this for speed at some point.
-	FRigMapperProcessor RigMapperProcessorTemp = RigMapperProcessor;
-	TSet<FName> OutputNames;
-
-	if (RigMapperProcessorTemp.IsValid())
-	{
-		RigMapperProcessorTemp.EvaluateFrames(InCurveMetaData.Names, InCurveFrameValues, OutCurveMetaData.Names, OutCurveFrameValues);
-		OutCurveMetaData.Flags.Init(4, OutCurveMetaData.Names.Num()); // a value of 4 means Editable curve
-		OutCurveMetaData.Colors.SetNum(OutCurveMetaData.Names.Num());
-		for (int32 Curve = 0; Curve < OutCurveMetaData.Names.Num(); ++Curve)
-		{
-			OutCurveMetaData.Colors[Curve] = MakeCurveColor(OutCurveMetaData.Names[Curve]);
-			OutputNames.Add(OutCurveMetaData.Names[Curve]);
-		}
-
-		// 'pass-through' any input curves which have not been modified, if bCopyAllSourceCurves is true
-		if (Settings.bCopyAllSourceCurves)
-		{
-			TArray<int32> InputCurveIndicesToAdd;
-			InputCurveIndicesToAdd.Reserve(InCurveMetaData.Names.Num());
-			for (int32 InputCurve = 0; InputCurve < InCurveMetaData.Names.Num(); ++InputCurve)
-			{
-				if (!OutputNames.Contains(InCurveMetaData.Names[InputCurve]))
-				{
-					InputCurveIndicesToAdd.Add(InputCurve);
-				}
-			}
-
-			// now add these curves
-			OutCurveMetaData.Names.Reserve(OutCurveMetaData.Names.Num() + InputCurveIndicesToAdd.Num());
-			OutCurveMetaData.Flags.Reserve(OutCurveMetaData.Flags.Num() + InputCurveIndicesToAdd.Num());
-			OutCurveMetaData.Colors.Reserve(OutCurveMetaData.Colors.Num() + InputCurveIndicesToAdd.Num());
-			int32 Frame = 0;
-			for (TArray<TOptional<float>>& OutputFrame : OutCurveFrameValues)
-			{
-				OutputFrame.Reserve(OutputFrame.Num() + InputCurveIndicesToAdd.Num());
-				for (int32 CurveIndex : InputCurveIndicesToAdd)
-				{
-					OutputFrame.Add(InCurveFrameValues[Frame][CurveIndex]);
-				}
-
-				Frame++;
-			}
-
-			for (int32 CurveIndex : InputCurveIndicesToAdd)
-			{
-				OutCurveMetaData.Names.Add(InCurveMetaData.Names[CurveIndex]);
-				OutCurveMetaData.Flags.Add(InCurveMetaData.Flags[CurveIndex]);
-				OutCurveMetaData.Colors.Add(InCurveMetaData.Colors[CurveIndex]);
-			}
-		}
+		bIsInitialized = Helper.InitializeRigMapping(DefinitionsToLoad);
 	}
 }
 
-bool FIKRetargetRigMapperOp::Initialize(
-	const FIKRetargetProcessor& InProcessor,
-	const FRetargetSkeleton& /*InSourceSkeleton*/,
-	const FTargetSkeleton& /*InTargetSkeleton*/,
-	const FIKRetargetOpBase* /*InParentOp*/,
-	FIKRigLogger& /*Log*/)
+void FIKRetargetRigMapperUserDataOp::AnimGraphEvaluateAnyThread(FPoseContext& Output)
 {
-	const FRetargetSkeleton& TargetSkeleton = InProcessor.GetSkeleton(ERetargetSourceOrTarget::Target);
-	bIsInitialized = InitializeRigMapping(TargetSkeleton.SkeletalMesh);
+	if (!IsEnabled())
+	{
+		return;
+	}
 
-	return bIsInitialized;
-};
+	FBlendedCurve& OutputCurves = Output.Curve;
+	FBlendedCurve InputCurves;
 
+	if (GetTakeInputCurvesFromSourceAnimInstance())
+	{
+		InputCurves.CopyFrom(SourceCurves);
+		if (Settings.bCopyAllSourceCurves)
+		{
+			OutputCurves.CopyFrom(SourceCurves);
+		}
+	}
+	else
+	{
+		InputCurves.CopyFrom(OutputCurves);
+		// clear the outputs if needed
+		if (!Settings.bCopyAllSourceCurves)
+		{
+			OutputCurves.ForEachElement([&OutputCurves](const UE::Anim::FCurveElement& InCurveElement)
+				{
+					OutputCurves.Set(InCurveElement.Name, 0.0f);
+				});
+		}
+	}
 
+	Helper.EvaluateRigMapping(InputCurves, OutputCurves);
+}
 
+TArray<URigMapperDefinition*> FIKRetargetRigMapperUserDataOp::GetDefinitionsToLoad(const USkeletalMesh* InTargetMesh)
+{
+	TArray<URigMapperDefinition*> DefinitionsToLoad;
+
+	URigMapperDefinitionUserData* UserData = FRigMapperOpHelper::GetUserDataFromMesh(InTargetMesh);
+	if (UserData)
+	{
+		DefinitionsToLoad = UserData->Definitions;
+	}
+
+	return DefinitionsToLoad;
+}
+
+const UClass* FIKRetargetRigMapperUserDataOpSettings::GetControllerType() const
+{
+	return UIKRetargetRigMapperUserDataOpController::StaticClass();
+}
+
+void FIKRetargetRigMapperUserDataOpSettings::CopySettingsAtRuntime(const FIKRetargetOpSettingsBase* InSettingsToCopyFrom)
+{
+	// copies all properties
+	const TArray<FName> PropertiesToIgnore = {};
+	FIKRetargetOpBase::CopyStructProperties(
+		FIKRetargetRigMapperUserDataOpSettings::StaticStruct(),
+		InSettingsToCopyFrom,
+		this,
+		PropertiesToIgnore);
+}
+
+void UIKRetargetRigMapperOpController::SetSettings(FIKRetargetRigMapperOpSettings InSettings)
+{
+	OpSettingsToControl->CopySettingsAtRuntime(&InSettings);
+}
+
 FIKRetargetRigMapperOpSettings UIKRetargetRigMapperOpController::GetSettings()
 {
 	return *reinterpret_cast<FIKRetargetRigMapperOpSettings*>(OpSettingsToControl);
 }
 
-void UIKRetargetRigMapperOpController::SetSettings(FIKRetargetRigMapperOpSettings InSettings)
+void UIKRetargetRigMapperUserDataOpController::SetSettings(FIKRetargetRigMapperUserDataOpSettings InSettings)
 {
 	OpSettingsToControl->CopySettingsAtRuntime(&InSettings);
 }
+
+FIKRetargetRigMapperUserDataOpSettings UIKRetargetRigMapperUserDataOpController::GetSettings()
+{
+	return *reinterpret_cast<FIKRetargetRigMapperUserDataOpSettings*>(OpSettingsToControl);
+}
diff --git a/Engine/Plugins/Experimental/Animation/RigMapperOp/Source/RigMapperOp/Public/RigMapperOp.h b/Engine/Plugins/Experimental/Animation/RigMapperOp/Source/RigMapperOp/Public/RigMapperOp.h
--- a/Engine/Plugins/Experimental/Animation/RigMapperOp/Source/RigMapperOp/Public/RigMapperOp.h
+++ b/Engine/Plugins/Experimental/Animation/RigMapperOp/Source/RigMapperOp/Public/RigMapperOp.h
@@ -2,7 +2,7 @@
 #pragma once
 
 #include "Retargeter/IKRetargetOps.h"
-#include "StructUtils/InstancedStruct.h"
+#include "Retargeter/RetargetOps/CurveRemapOp.h"
 #include "Animation/AnimCurveTypes.h"
 #include "RigMapperProcessor.h"
 
@@ -17,8 +17,8 @@
 class USkeletalMesh;
 class USkeletalMeshComponent;
 
-USTRUCT(BlueprintType, meta = (DisplayName = "RigMapper Settings"))
-struct FIKRetargetRigMapperOpSettings : public FIKRetargetOpSettingsBase
+USTRUCT(BlueprintType)
+struct FIKRetargetRigMapperOpSettingsBase : public FIKRetargetOpSettingsBase
 {
 	GENERATED_BODY()
 
@@ -30,89 +30,66 @@
 	// and the RigMapper covers only a subset of the controls, but false otherwise.
 	UPROPERTY(EditAnywhere, BlueprintReadWrite, Category = "Copy Curves")
 	bool bCopyAllSourceCurves = false;
+};
 
+USTRUCT(BlueprintType, meta = (DisplayName = "Single Definition RigMapper Settings"))
+struct FIKRetargetRigMapperOpSettings : public FIKRetargetRigMapperOpSettingsBase
+{
+	GENERATED_BODY();
+
+	// User defined RigMapper definition
 	UPROPERTY(EditAnywhere, BlueprintReadWrite, Category = RigMapper)
+	TObjectPtr<URigMapperDefinition> Definition;
+
+	// Whether should RigMapper definition be overridden with the array of definitions from target skeletal mesh UserData
+	// NOTE: Should not be used if there are multiple single RigMapper ops in the sequence
+	UPROPERTY(EditAnywhere, BlueprintReadWrite, Category = RigMapper)
+	bool bOverrideFromUserDataDefinitions = false;
+
+	PRAGMA_DISABLE_DEPRECATION_WARNINGS
+	// TODO: Remove in UE 5.8
+	UE_DEPRECATED(5.7, "Use Definition property instead. This array will be removed in UE 5.8")
+	UPROPERTY(meta = (DeprecatedProperty, DeprecationMessage = "Use Definition instead"))
 	TArray<TObjectPtr<URigMapperDefinition>> Definitions;
+	PRAGMA_ENABLE_DEPRECATION_WARNINGS
 
 	UE_API virtual const UClass* GetControllerType() const override;
-	
 	UE_API virtual void CopySettingsAtRuntime(const FIKRetargetOpSettingsBase* InSettingsToCopyFrom) override;
 };
 
-USTRUCT(BlueprintType, meta = (DisplayName = "RigMapper"))
-struct FIKRetargetRigMapperOp : public FIKRetargetOpBase
+USTRUCT(BlueprintType, meta = (DisplayName = "UserData RigMapper Settings"))
+struct FIKRetargetRigMapperUserDataOpSettings : public FIKRetargetRigMapperOpSettingsBase
 {
-	GENERATED_BODY()
-	
-	
-	virtual bool Initialize(
-		const FIKRetargetProcessor& InProcessor,
-		const FRetargetSkeleton& InSourceSkeleton,
-		const FTargetSkeleton& InTargetSkeleton,
-		const FIKRetargetOpBase* InParentOp,
-		FIKRigLogger& Log) override;
-	
-	// NOTE: this op does not do anything in Run().
-	virtual void Run(
-		FIKRetargetProcessor& InProcessor,
-		const double InDeltaTime,
-		const TArray<FTransform>& InSourceGlobalPose,
-		TArray<FTransform>& OutTargetGlobalPose) override {};
+	GENERATED_BODY();
 
-	virtual FIKRetargetOpSettingsBase* GetSettings() override
-	{
-		return &Settings;
-	};
+	UE_API virtual const UClass* GetControllerType() const override;
+	UE_API virtual void CopySettingsAtRuntime(const FIKRetargetOpSettingsBase* InSettingsToCopyFrom) override;
+};
 
-	virtual void SetSettings(const FIKRetargetOpSettingsBase* InSettings) override
-	{
-		Settings = *reinterpret_cast<const FIKRetargetRigMapperOpSettings*>(InSettings);
-	};
-	
-	virtual const UScriptStruct* GetSettingsType() const override
-	{
-		return FIKRetargetRigMapperOpSettings::StaticStruct();
-	}
-	
-	virtual const UScriptStruct* GetType() const override
-	{
-		return FIKRetargetRigMapperOp::StaticStruct();
-	}
+// A helper struct containing business logic common for all RigMapper ops
+struct RIGMAPPEROP_API FRigMapperOpHelper
+{
+public:
+	FRigMapperOpHelper() = default;
 
-	/* RigMapperOp is a singleton because it can override definitions from the Target SkelMesh user data and therefore we should only do that once*/
-	virtual bool IsSingleton() const override { return true; };
+	bool InitializeRigMapping(const TArray<URigMapperDefinition*>& InDefinitions);
 
-	virtual bool HasCurveProcessing() const override 
-	{ 
-		return true; 
-	}
+	void ProcessAnimSequenceCurves(FIKRetargetOpBase::FCurveData InCurveMetaData, FIKRetargetOpBase::FFrameValues InCurveFrameValues,
+		FIKRetargetOpBase::FCurveData& OutCurveMetaData, FIKRetargetOpBase::FFrameValues& OutCurveFrameValues, bool bCopyAllSourceCurves) const;
 
-	virtual void ProcessAnimSequenceCurves(FIKRetargetOpBase::FCurveData InCurveMetaData, FIKRetargetOpBase::FFrameValues InCurveFrameValues,
-		FIKRetargetOpBase::FCurveData& OutCurveMetaData, FIKRetargetOpBase::FFrameValues& OutCurveFrameValues)const override;
-
-
-	UE_API virtual void AnimGraphPreUpdateMainThread(USkeletalMeshComponent& SourceMeshComponent, USkeletalMeshComponent& TargetMeshComponent) override;
-
-	UE_API virtual void AnimGraphEvaluateAnyThread(FPoseContext& Output) override;
-
-
-	UPROPERTY()
-	FIKRetargetRigMapperOpSettings Settings;
-	
-	// cached curves, copied on the game thread in PreUpdate()
-	FBlendedHeapCurve SourceCurves;
-
-
-private:
-	bool InitializeRigMapping(const USkeletalMesh* InTargetMesh);
 	void EvaluateRigMapping(const FBlendedCurve& InCurve, FBlendedCurve& OutCurve);
 
+	bool CheckReInit(const TArray<URigMapperDefinition*>& InCurrentDefinitions);
+
+	bool IsValid() const { return RigMapperProcessor.IsValid(); }
+
+	static URigMapperDefinitionUserData* GetUserDataFromMesh(const USkeletalMesh* InMesh);
+
+private:
+
 	// The definitions that we have loaded. Cached to check against changes and reinit if need be
 	TArray<TObjectPtr<URigMapperDefinition>> LoadedDefinitions;
 
-	// The asset user data currently used to override definitions, if any was set on the skeletal mesh
-	TObjectPtr<URigMapperDefinitionUserData> LoadedUserData;
-
 	// The processor to evaluate the rig mapping
 	FRigMapperProcessor RigMapperProcessor;
 
@@ -130,7 +107,8 @@
 		FRigMapperCurveMapping(FName InName, int32 InCurveIndex)
 			: Name(InName)
 			, CurveIndex(InCurveIndex)
-		{}
+		{
+		}
 
 		FName Name = NAME_None;
 		int32 CurveIndex = INDEX_NONE;
@@ -144,7 +122,8 @@
 		FRigMapperOutputCurveMapping(FName InName, int32 InOutputCurveIndex, int32 InInputCurveIndex)
 			: FRigMapperCurveMapping(InName, InOutputCurveIndex)
 			, InputCurveIndex(InInputCurveIndex)
-		{}
+		{
+		}
 
 		int32 InputCurveIndex = INDEX_NONE;
 	};
@@ -156,7 +135,76 @@
 	// Output curve mapping for the bulk get of the curve values
 	using FOutputCurveMappings = UE::Anim::TNamedValueArray<FDefaultAllocator, FRigMapperOutputCurveMapping>;
 	FOutputCurveMappings OutputCurveMappings;
-
 };
 
+// Single definition RigMapper op with an option to use UserData definitions array
+USTRUCT(BlueprintType, meta = (DisplayName = "Single RigMapper"))
+struct FIKRetargetRigMapperOp : public FIKRetargetOpBase
+{
+	GENERATED_BODY()
+	
+public:	
+	virtual bool Initialize(
+		const FIKRetargetProcessor& InProcessor,
+		const FRetargetSkeleton& InSourceSkeleton,
+		const FTargetSkeleton& InTargetSkeleton,
+		const FIKRetargetOpBase* InParentOp,
+		FIKRigLogger& Log) override;
+
+	// NOTE: RigMapper ops do not do anything in Run().
+	virtual void Run(
+		FIKRetargetProcessor& InProcessor,
+		const double InDeltaTime,
+		const TArray<FTransform>& InSourceGlobalPose,
+		TArray<FTransform>& OutTargetGlobalPose) override {};
+
+	virtual FIKRetargetOpSettingsBase* GetSettings() override
+	{
+		return &Settings;
+	};
+
+	virtual void SetSettings(const FIKRetargetOpSettingsBase* InSettings) override
+	{
+		Settings = *reinterpret_cast<const FIKRetargetRigMapperOpSettings*>(InSettings);
+	};
+	
+	virtual const UScriptStruct* GetSettingsType() const override
+	{
+		return FIKRetargetRigMapperOpSettings::StaticStruct();
+	}
+
+	virtual const UScriptStruct* GetParentOpType() const override
+	{
+		return FIKRetargetCurveRemapOp::StaticStruct();
+	}
+	
+	virtual const UScriptStruct* GetType() const override
+	{
+		return FIKRetargetRigMapperOp::StaticStruct();
+	}
+
+	virtual bool IsSingleton() const override { return false; }
+
+	virtual void ProcessAnimSequenceCurves(FIKRetargetOpBase::FCurveData InCurveMetaData, FIKRetargetOpBase::FFrameValues InCurveFrameValues,
+		FIKRetargetOpBase::FCurveData& OutCurveMetaData, FIKRetargetOpBase::FFrameValues& OutCurveFrameValues) const;
+
+	UE_API virtual void AnimGraphPreUpdateMainThread(USkeletalMeshComponent& SourceMeshComponent, USkeletalMeshComponent& TargetMeshComponent) override;
+
+	UE_API virtual void AnimGraphEvaluateAnyThread(FPoseContext& Output) override;
+
+	bool PostLoadSplitOp(TArray<FInstancedStruct>& OutNewOps);
+
+	UPROPERTY()
+	FIKRetargetRigMapperOpSettings Settings;
+
+	// Cached curves, copied on the game thread in PreUpdate()
+	FBlendedHeapCurve SourceCurves;
+
+protected:
+	TArray<URigMapperDefinition*> GetDefinitionsToLoad(const USkeletalMesh* InTargetMesh);
+
+	// A helper struct containing business logic common for all RigMapper ops
+	FRigMapperOpHelper Helper;
+};
+
 /* The blueprint/python API for editing a RigMapper Op */
@@ -177,6 +225,93 @@
 	UE_API void SetSettings(FIKRetargetRigMapperOpSettings InSettings);
 };
 
+// UserData definitions ONLY RigMapper op
+USTRUCT(BlueprintType, meta = (DisplayName = "UserData RigMapper"))
+struct FIKRetargetRigMapperUserDataOp : public FIKRetargetOpBase
+{
+	GENERATED_BODY()
 
+public:
+	virtual bool Initialize(
+		const FIKRetargetProcessor& InProcessor,
+		const FRetargetSkeleton& InSourceSkeleton,
+		const FTargetSkeleton& InTargetSkeleton,
+		const FIKRetargetOpBase* InParentOp,
+		FIKRigLogger& Log) override;
+
+	// NOTE: RigMapper ops do not do anything in Run().
+	virtual void Run(
+		FIKRetargetProcessor& InProcessor,
+		const double InDeltaTime,
+		const TArray<FTransform>& InSourceGlobalPose,
+		TArray<FTransform>& OutTargetGlobalPose) override {
+	};
+
+	virtual FIKRetargetOpSettingsBase* GetSettings() override
+	{
+		return &Settings;
+	};
+
+	virtual void SetSettings(const FIKRetargetOpSettingsBase* InSettings) override
+	{
+		Settings = *reinterpret_cast<const FIKRetargetRigMapperUserDataOpSettings*>(InSettings);
+	};
+
+	virtual const UScriptStruct* GetSettingsType() const override
+	{
+		return FIKRetargetRigMapperUserDataOpSettings::StaticStruct();
+	}
+
+	virtual const UScriptStruct* GetParentOpType() const override
+	{
+		return FIKRetargetCurveRemapOp::StaticStruct();
+	}
+
+	virtual const UScriptStruct* GetType() const override
+	{
+		return FIKRetargetRigMapperUserDataOp::StaticStruct();
+	}
+
+	virtual bool IsSingleton() const override { return false; }
+
+	virtual void ProcessAnimSequenceCurves(FIKRetargetOpBase::FCurveData InCurveMetaData, FIKRetargetOpBase::FFrameValues InCurveFrameValues,
+		FIKRetargetOpBase::FCurveData& OutCurveMetaData, FIKRetargetOpBase::FFrameValues& OutCurveFrameValues) const;
+
+	UE_API virtual void AnimGraphPreUpdateMainThread(USkeletalMeshComponent& SourceMeshComponent, USkeletalMeshComponent& TargetMeshComponent) override;
+
+	UE_API virtual void AnimGraphEvaluateAnyThread(FPoseContext& Output) override;
+
+	// Cached curves, copied on the game thread in PreUpdate()
+	FBlendedHeapCurve SourceCurves;
+
+	UPROPERTY()
+	FIKRetargetRigMapperUserDataOpSettings Settings;
+
+protected:
+	TArray<URigMapperDefinition*> GetDefinitionsToLoad(const USkeletalMesh* InTargetMesh);
+
+	// A helper struct containing business logic common for all RigMapper ops
+	FRigMapperOpHelper Helper;
+};
+
+/* The blueprint/python API for editing a RigMapper UserData Op */
+UCLASS(MinimalAPI, BlueprintType)
+class UIKRetargetRigMapperUserDataOpController : public UIKRetargetOpControllerBase
+{
+	GENERATED_BODY()
+
+public:
+	/* Get the current op settings as a struct.
+	 * @return FIKRetargetRigMapperUserDataOpSettings struct with the current settings used by the op. */
+	UFUNCTION(BlueprintCallable, Category = Settings)
+	UE_API FIKRetargetRigMapperUserDataOpSettings GetSettings();
+
+	/* Set the op settings. Input is a custom struct type for this op.
+	 * @param InSettings a FIKRetargetRigMapperUserDataOpSettings struct containing all the settings to apply to this op */
+	UFUNCTION(BlueprintCallable, Category = Settings)
+	UE_API void SetSettings(FIKRetargetRigMapperUserDataOpSettings InSettings);
+
+};
+
 #undef LOCTEXT_NAMESPACE
 
