BuildHealth: RigMapperOp deprecation warnings disabled

#Synced-CL 49271205

diff --git a/Engine/Plugins/Experimental/Animation/RigMapperOp/Source/RigMapperOp/Private/RigMapperOp.cpp b/Engine/Plugins/Experimental/Animation/RigMapperOp/Source/RigMapperOp/Private/RigMapperOp.cpp
--- a/Engine/Plugins/Experimental/Animation/RigMapperOp/Source/RigMapperOp/Private/RigMapperOp.cpp
+++ b/Engine/Plugins/Experimental/Animation/RigMapperOp/Source/RigMapperOp/Private/RigMapperOp.cpp
@@ -266,7 +266,6 @@ void FIKRetargetRigMapperOp::AnimGraphEvaluateAnyThread(FPoseContext& Output)
 
 bool FIKRetargetRigMapperOp::PostLoadSplitOp(TArray<FInstancedStruct>& OutNewOps)
 {
-	PRAGMA_DISABLE_DEPRECATION_WARNINGS
 	// Only split if there are multiple definitions in the old RigMapperOp
 	if (Settings.Definitions.Num() <= 1)
 	{
@@ -312,7 +311,6 @@ bool FIKRetargetRigMapperOp::PostLoadSplitOp(TArray<FInstancedStruct>& OutNewOps
 	}
 
 	return true; // Replace original op with new ops
-	PRAGMA_ENABLE_DEPRECATION_WARNINGS
 }
 
 TArray<URigMapperDefinition*> FIKRetargetRigMapperOp::GetDefinitionsToLoad(const USkeletalMesh* InTargetMesh)
diff --git a/Engine/Plugins/Experimental/Animation/RigMapperOp/Source/RigMapperOp/Public/RigMapperOp.h b/Engine/Plugins/Experimental/Animation/RigMapperOp/Source/RigMapperOp/Public/RigMapperOp.h
--- a/Engine/Plugins/Experimental/Animation/RigMapperOp/Source/RigMapperOp/Public/RigMapperOp.h
+++ b/Engine/Plugins/Experimental/Animation/RigMapperOp/Source/RigMapperOp/Public/RigMapperOp.h
@@ -46,12 +46,9 @@ struct FIKRetargetRigMapperOpSettings : public FIKRetargetRigMapperOpSettingsBas
 	UPROPERTY(EditAnywhere, BlueprintReadWrite, Category = RigMapper)
 	bool bOverrideFromUserDataDefinitions = false;
 
-	PRAGMA_DISABLE_DEPRECATION_WARNINGS
 	// TODO: Remove in UE 5.8
-	UE_DEPRECATED(5.7, "Use Definition property instead. This array will be removed in UE 5.8")
 	UPROPERTY(meta = (DeprecatedProperty, DeprecationMessage = "Use Definition instead"))
 	TArray<TObjectPtr<URigMapperDefinition>> Definitions;
-	PRAGMA_ENABLE_DEPRECATION_WARNINGS
 
 	UE_API virtual const UClass* GetControllerType() const override;
 	UE_API virtual void CopySettingsAtRuntime(const FIKRetargetOpSettingsBase* InSettingsToCopyFrom) override;
