[IKRig] Moved curve manipulation specific operations into a new abstract IKRetargetCurveOpBase struct

#Synced-CL 49404572

diff --git a/Engine/Plugins/Animation/IKRig/Source/IKRig/Private/Retargeter/IKRetargetProcessor.cpp b/Engine/Plugins/Animation/IKRig/Source/IKRig/Private/Retargeter/IKRetargetProcessor.cpp
--- a/Engine/Plugins/Animation/IKRig/Source/IKRig/Private/Retargeter/IKRetargetProcessor.cpp
+++ b/Engine/Plugins/Animation/IKRig/Source/IKRig/Private/Retargeter/IKRetargetProcessor.cpp
@@ -13,6 +13,7 @@
 #include "Engine/SkeletalMesh.h"
 #include "Retargeter/RetargetOps/CopyBasePoseOp.h"
 #include "Retargeter/RetargetOps/ScaleSourceOp.h"
+#include "Retargeter/RetargetOps/RetargetCurvesOp.h"
 
 #include UE_INLINE_GENERATED_CPP_BY_NAME(IKRetargetProcessor)
 
@@ -1676,10 +1677,17 @@ void FIKRetargetProcessor::OnAnimGraphEvaluateAnyThread(FPoseContext& Output)
 	bool bFirstCurveOp = true;
 	for (FInstancedStruct& OpStruct : OpStack)
 	{
+		const UStruct* StructType = OpStruct.GetScriptStruct();
 		FIKRetargetOpBase& Op = OpStruct.GetMutable<FIKRetargetOpBase>();
-		Op.SetTakeInputCurvesFromSourceAnimInstance(bFirstCurveOp);
+		bool bHasCurveProcessing = false;
+		if (StructType && StructType->IsChildOf(FIKRetargetCurvesOpBase::StaticStruct()))
+		{
+			FIKRetargetCurvesOpBase& CurveOp = OpStruct.GetMutable<FIKRetargetCurvesOpBase>();
+			CurveOp.SetTakeInputCurvesFromSourceAnimInstance(bFirstCurveOp);
+			bHasCurveProcessing = CurveOp.IsEnabled() && CurveOp.HasCurveProcessing();
+		}
 		Op.AnimGraphEvaluateAnyThread(Output);
-		if (Op.IsEnabled() && Op.HasCurveProcessing())
+		if (bHasCurveProcessing)
 		{
 			bFirstCurveOp = false;
 		}
diff --git a/Engine/Plugins/Animation/IKRig/Source/IKRig/Private/Retargeter/RetargetOps/CurveRemapOp.cpp b/Engine/Plugins/Animation/IKRig/Source/IKRig/Private/Retargeter/RetargetOps/CurveRemapOp.cpp
--- a/Engine/Plugins/Animation/IKRig/Source/IKRig/Private/Retargeter/RetargetOps/CurveRemapOp.cpp
+++ b/Engine/Plugins/Animation/IKRig/Source/IKRig/Private/Retargeter/RetargetOps/CurveRemapOp.cpp
@@ -124,11 +124,11 @@ void FIKRetargetCurveRemapOp::AnimGraphEvaluateAnyThread(FPoseContext& Output)
 	}
 }
 
-void FIKRetargetCurveRemapOp::ProcessAnimSequenceCurves(FIKRetargetOpBase::FCurveData InCurveMetaData, FIKRetargetOpBase::FFrameValues InCurveFrameValues,
-	FIKRetargetOpBase::FCurveData& OutCurveMetaData, FIKRetargetOpBase::FFrameValues& OutCurveFrameValues) const
+void FIKRetargetCurveRemapOp::ProcessAnimSequenceCurves(FIKRetargetCurvesOpBase::FCurveData InCurveMetaData, FIKRetargetCurvesOpBase::FFrameValues InCurveFrameValues,
+	FIKRetargetCurvesOpBase::FCurveData& OutCurveMetaData, FIKRetargetCurvesOpBase::FFrameValues& OutCurveFrameValues) const
 {
-	FIKRetargetOpBase::FCurveData SourceCurveMetaData;
-	FIKRetargetOpBase::FFrameValues SourceCurveFrameValues;
+	FIKRetargetCurvesOpBase::FCurveData SourceCurveMetaData;
+	FIKRetargetCurvesOpBase::FFrameValues SourceCurveFrameValues;
 	TArray<FName> ValidSourceTargetNames;
 
 	if (Settings.bRemapCurves)
diff --git a/Engine/Plugins/Animation/IKRig/Source/IKRig/Private/Retargeter/RetargetOps/RetargetCurvesOp.cpp b/Engine/Plugins/Animation/IKRig/Source/IKRig/Private/Retargeter/RetargetOps/RetargetCurvesOp.cpp
new file mode 100644
--- /dev/null
+++ b/Engine/Plugins/Animation/IKRig/Source/IKRig/Private/Retargeter/RetargetOps/RetargetCurvesOp.cpp
@@ -0,0 +1,7 @@
+// Copyright Epic Games, Inc. All Rights Reserved.
+
+#include "Retargeter/RetargetOps/RetargetCurvesOp.h"
+
+#include UE_INLINE_GENERATED_CPP_BY_NAME(RetargetCurvesOp)
+
+// TODO: For future curve manpulation base struct implementation.
\ No newline at end of file
diff --git a/Engine/Plugins/Animation/IKRig/Source/IKRig/Public/Retargeter/IKRetargetOps.h b/Engine/Plugins/Animation/IKRig/Source/IKRig/Public/Retargeter/IKRetargetOps.h
--- a/Engine/Plugins/Animation/IKRig/Source/IKRig/Public/Retargeter/IKRetargetOps.h
+++ b/Engine/Plugins/Animation/IKRig/Source/IKRig/Public/Retargeter/IKRetargetOps.h
@@ -215,41 +215,6 @@ public:
 
 	/** (optional) called during the PostLoad of the owning asset. Can be used to apply patches to the data for backwards compatibility. */
 	UE_API virtual void PostLoad(const FIKRigObjectVersion::Type InVersion);
-
-	using FFrameValues = TArray<TArray<TOptional<float>>>;
-
-	struct FCurveData
-	{
-		TArray<FName> Names;
-		TArray<int32> Flags;
-		TArray<FLinearColor> Colors;
-	};
-
-	/** (optional) a function which can be used during batch processing to apply any operations to curve data for each frame of an anim sequence. By default, this just moves the input data to the output data*/
-	virtual void ProcessAnimSequenceCurves(FIKRetargetOpBase::FCurveData InCurveMetaData, FIKRetargetOpBase::FFrameValues InCurveFrameValues,
-		FIKRetargetOpBase::FCurveData& OutCurveMetaData, FIKRetargetOpBase::FFrameValues& OutCurveFrameValues) const
-	{
-		OutCurveMetaData = MoveTemp(InCurveMetaData);
-		OutCurveFrameValues = MoveTemp(InCurveFrameValues);
-	}
-
-	/** (optional) a function which returns true if the operation performs curve processing, false otherwise (default is false) */
-	virtual bool HasCurveProcessing() const
-	{
-		return false;
-	}
-
-	/** a function which allows the user to manually indicate whether to take source curves from the source anim instance, or from the current pose context, since only the first enabled curve processing node must take inputs from the source anim instance*/
-	void SetTakeInputCurvesFromSourceAnimInstance(bool bInTakeInputCurvesFromSourceAnimInstance)
-	{
-		bTakeInputCurvesFromSourceAnimInstance = bInTakeInputCurvesFromSourceAnimInstance;
-	}
-
-	/** a function which allows the user to manually indicate whether to take source curves from the source anim instance, or from the current pose context, since only the first enabled curve processing node must take inputs from the source anim instance */
-	bool GetTakeInputCurvesFromSourceAnimInstance() const
-	{
-		return bTakeInputCurvesFromSourceAnimInstance;
-	}
 	
 	/** (optional) override to automate initial setup after being added to the stack */
 	virtual void OnAddedToStack(const UIKRetargeter* InRetargetAsset, const FIKRetargetOpBase* InParentOp) {};
@@ -338,6 +303,9 @@ public:
 	/** (optional) return true to disallow multiple copies of this op in the stack */
 	virtual bool IsSingleton() const { return false; };
 
+	/** (optional) return true to disallow adding this op to the stack */
+	virtual bool IsAbstract() const { return false; };
+
 	/** (optional) override and supply the target IK Rig this op references
 	 * NOTE: The retarget processor will resolve the bone chains in the rig so that they can be queried by the op. */
 	virtual const UIKRigDefinition* GetCustomTargetIKRig() const { return nullptr; };
@@ -438,6 +406,15 @@ public:
 		void* InOutDestStruct,
 		const TArray<FName>& InPropertiesToIgnore);
 
+	using FFrameValues = TArray<TArray<TOptional<float>>>;
+
+	struct FCurveData
+	{
+		TArray<FName> Names;
+		TArray<int32> Flags;
+		TArray<FLinearColor> Colors;
+	};
+
 private:
 	
 	/** when false, execution of this op is skipped */
@@ -452,10 +429,6 @@ private:
 	UPROPERTY()
 	FName ParentOpName;
 
-	/** (optional) determines whether input curves are taken from the source anim instance or from the current pose context (for example when we have multiple nodes which process curves). Can be set on the fly */
-	UPROPERTY(Transient)
-	bool bTakeInputCurvesFromSourceAnimInstance = true;
-
 protected:
 	
 	/** set to true when the op is ready to run */
diff --git a/Engine/Plugins/Animation/IKRig/Source/IKRig/Public/Retargeter/RetargetOps/CurveRemapOp.h b/Engine/Plugins/Animation/IKRig/Source/IKRig/Public/Retargeter/RetargetOps/CurveRemapOp.h
--- a/Engine/Plugins/Animation/IKRig/Source/IKRig/Public/Retargeter/RetargetOps/CurveRemapOp.h
+++ b/Engine/Plugins/Animation/IKRig/Source/IKRig/Public/Retargeter/RetargetOps/CurveRemapOp.h
@@ -2,6 +2,7 @@
 #pragma once
 
 #include "Retargeter/IKRetargetOps.h"
+#include "Retargeter/RetargetOps/RetargetCurvesOp.h"
 #include "StructUtils/InstancedStruct.h"
 #include "Animation/AnimCurveTypes.h"
 
@@ -54,7 +55,7 @@ struct FIKRetargetCurveRemapOpSettings : public FIKRetargetOpSettingsBase
 };
 
 USTRUCT(BlueprintType, meta = (DisplayName = "Remap Curves"))
-struct FIKRetargetCurveRemapOp : public FIKRetargetOpBase
+struct FIKRetargetCurveRemapOp : public FIKRetargetCurvesOpBase
 {
 	GENERATED_BODY()
 
@@ -105,6 +106,8 @@ public:
 		return true;
 	}
 
+	virtual bool IsAbstract() const override { return false; };
+
 	virtual bool IsSingleton() const override { return true; };
 
 	virtual bool HasCurveProcessing() const override
@@ -112,15 +115,12 @@ public:
 		return true;
 	}
 
-	virtual void ProcessAnimSequenceCurves(FIKRetargetOpBase::FCurveData InCurveMetaData, FIKRetargetOpBase::FFrameValues InCurveFrameValues,
-		FIKRetargetOpBase::FCurveData& OutCurveMetaData, FIKRetargetOpBase::FFrameValues& OutCurveFrameValues)const override;
+	virtual void ProcessAnimSequenceCurves(FIKRetargetCurvesOpBase::FCurveData InCurveMetaData, FIKRetargetCurvesOpBase::FFrameValues InCurveFrameValues,
+		FIKRetargetCurvesOpBase::FCurveData& OutCurveMetaData, FIKRetargetCurvesOpBase::FFrameValues& OutCurveFrameValues)const override;
 
 	UE_API virtual void AnimGraphPreUpdateMainThread(USkeletalMeshComponent& SourceMeshComponent, USkeletalMeshComponent& TargetMeshComponent) override;
 
 	UE_API virtual void AnimGraphEvaluateAnyThread(FPoseContext& Output) override;
-	
-	// cached curves, copied on the game thread in PreUpdate()
-	FBlendedHeapCurve SourceCurves;
 
 	UPROPERTY()
 	FIKRetargetCurveRemapOpSettings Settings;
diff --git a/Engine/Plugins/Animation/IKRig/Source/IKRig/Public/Retargeter/RetargetOps/RetargetCurvesOp.h b/Engine/Plugins/Animation/IKRig/Source/IKRig/Public/Retargeter/RetargetOps/RetargetCurvesOp.h
new file mode 100644
--- /dev/null
+++ b/Engine/Plugins/Animation/IKRig/Source/IKRig/Public/Retargeter/RetargetOps/RetargetCurvesOp.h
@@ -0,0 +1,94 @@
+// Copyright Epic Games, Inc. All Rights Reserved.
+#pragma once
+
+#include "Retargeter/IKRetargetOps.h"
+#include "StructUtils/InstancedStruct.h"
+#include "Animation/AnimCurveTypes.h"
+
+#include "RetargetCurvesOp.generated.h"
+
+#define UE_API IKRIG_API
+
+#define LOCTEXT_NAMESPACE "RetargetCurvesOp"
+
+/**
+ * This is the base class for defining all animation curves operations including curve remapping and rig mapping.
+*/
+USTRUCT(BlueprintType)
+struct FIKRetargetCurvesOpBase : public FIKRetargetOpBase
+{
+	GENERATED_BODY()
+
+public:
+	// NOTE: this op does not do anything in Initialize() or Run().
+	// In the future we may remove this coupling and make the op itself do the work via callbacks
+	virtual bool Initialize(
+		const FIKRetargetProcessor& InProcessor,
+		const FRetargetSkeleton& InSourceSkeleton,
+		const FTargetSkeleton& InTargetSkeleton,
+		const FIKRetargetOpBase* InParentOp,
+		FIKRigLogger& Log) override
+	{
+		bIsInitialized = true;
+		return true;
+	};
+
+	virtual void Run(
+		FIKRetargetProcessor& InProcessor,
+		const double InDeltaTime,
+		const TArray<FTransform>& InSourceGlobalPose,
+		TArray<FTransform>& OutTargetGlobalPose) override {
+	};
+
+	virtual const UScriptStruct* GetType() const override { return FIKRetargetCurvesOpBase::StaticStruct(); }
+
+	bool CanHaveChildOps() const override { return true; }
+
+	virtual bool IsSingleton() const override { return true; };
+
+	// Designator of an "abstract" base struct, in this case used to classify curve related ops
+	virtual bool IsAbstract() const override { return true; };
+
+	/** (optional) a function which returns true if the operation performs curve processing, false otherwise (default is true) */
+	virtual bool HasCurveProcessing() const { return true;  }
+
+	using FFrameValues = TArray<TArray<TOptional<float>>>;
+
+	struct FCurveData
+	{
+		TArray<FName> Names;
+		TArray<int32> Flags;
+		TArray<FLinearColor> Colors;
+	};
+
+	/** (optional) a function which can be used during batch processing to apply any operations to curve data for each frame of an anim sequence. By default, this just moves the input data to the output data*/
+	virtual void ProcessAnimSequenceCurves(FIKRetargetCurvesOpBase::FCurveData InCurveMetaData, FIKRetargetCurvesOpBase::FFrameValues InCurveFrameValues,
+		FIKRetargetCurvesOpBase::FCurveData& OutCurveMetaData, FIKRetargetCurvesOpBase::FFrameValues& OutCurveFrameValues) const
+	{
+		OutCurveMetaData = MoveTemp(InCurveMetaData);
+		OutCurveFrameValues = MoveTemp(InCurveFrameValues);
+	}
+
+	/** a function which allows the user to manually indicate whether to take source curves from the source anim instance, or from the current pose context, since only the first enabled curve processing node must take inputs from the source anim instance*/
+	void SetTakeInputCurvesFromSourceAnimInstance(bool bInTakeInputCurvesFromSourceAnimInstance)
+	{
+		bTakeInputCurvesFromSourceAnimInstance = bInTakeInputCurvesFromSourceAnimInstance;
+	}
+
+	/** a function which allows the user to manually indicate whether to take source curves from the source anim instance, or from the current pose context, since only the first enabled curve processing node must take inputs from the source anim instance */
+	bool GetTakeInputCurvesFromSourceAnimInstance() const
+	{
+		return bTakeInputCurvesFromSourceAnimInstance;
+	}
+
+	// cached curves, copied on the game thread in PreUpdate()
+	FBlendedHeapCurve SourceCurves;
+
+protected:
+	/** (optional) determines whether input curves are taken from the source anim instance or from the current pose context (for example when we have multiple nodes which process curves). Can be set on the fly */
+	UPROPERTY(Transient)
+	bool bTakeInputCurvesFromSourceAnimInstance = true;
+};
+
+#undef LOCTEXT_NAMESPACE
+#undef UE_API
\ No newline at end of file
diff --git a/Engine/Plugins/Animation/IKRig/Source/IKRigEditor/Private/RetargetEditor/IKRetargetBatchOperation.cpp b/Engine/Plugins/Animation/IKRig/Source/IKRigEditor/Private/RetargetEditor/IKRetargetBatchOperation.cpp
--- a/Engine/Plugins/Animation/IKRig/Source/IKRigEditor/Private/RetargetEditor/IKRetargetBatchOperation.cpp
+++ b/Engine/Plugins/Animation/IKRig/Source/IKRigEditor/Private/RetargetEditor/IKRetargetBatchOperation.cpp
@@ -536,8 +536,8 @@ void UIKRetargetBatchOperation::ApplyCurveOps(const FIKRetargetBatchOperationCon
 			for (const FInstancedStruct& OpStruct : InOpStack)
 			{
 				const UStruct* StructType = OpStruct.GetScriptStruct();
-				if (StructType && StructType->IsChildOf(FIKRetargetOpBase::StaticStruct()) &&
-					OpStruct.Get<FIKRetargetOpBase>().HasCurveProcessing())
+				if (StructType && StructType->IsChildOf(FIKRetargetCurvesOpBase::StaticStruct()) &&
+					OpStruct.Get<FIKRetargetCurvesOpBase>().HasCurveProcessing())
 				{
 					return false;
 				}
@@ -574,22 +574,22 @@ void UIKRetargetBatchOperation::ApplyCurveOps(const FIKRetargetBatchOperationCon
 		constexpr bool bShouldTransact = false;
 		TargetSeqController.OpenBracket(FText::FromString("Applying Curve Operations"), bShouldTransact);
 
-		FIKRetargetOpBase::FFrameValues InputCurveFrameValues;
-		FIKRetargetOpBase::FCurveData InputCurveMetaData;
+		FIKRetargetCurvesOpBase::FFrameValues InputCurveFrameValues;
+		FIKRetargetCurvesOpBase::FCurveData InputCurveMetaData;
 
 		TArray<FFrameTime> InputFrameTimes;
 		bool bLoadedData = FCurveUtils::LoadCurveValuesFromAnimSequence(SourceSequence, SourceSkeleton, InputCurveMetaData.Names, InputFrameTimes, InputCurveFrameValues, InputCurveMetaData.Flags, InputCurveMetaData.Colors);
 		if (bLoadedData)
 		{
-			FIKRetargetOpBase::FFrameValues OutputCurveFrameValues;
-			FIKRetargetOpBase::FCurveData OutputCurveMetaData;
+			FIKRetargetCurvesOpBase::FFrameValues OutputCurveFrameValues;
+			FIKRetargetCurvesOpBase::FCurveData OutputCurveMetaData;
 
 			for (const FInstancedStruct& OpStruct : InProcessor.GetRetargetOps())
 			{
 				const UStruct* StructType = OpStruct.GetScriptStruct();
-				if (StructType)
+				if (StructType && StructType->IsChildOf(FIKRetargetCurvesOpBase::StaticStruct()))
 				{
-					const FIKRetargetOpBase& Op = OpStruct.Get<FIKRetargetOpBase>();
+					const FIKRetargetCurvesOpBase& Op = OpStruct.Get<FIKRetargetCurvesOpBase>();
 					if (Op.IsEnabled())
 					{
 						Op.ProcessAnimSequenceCurves(InputCurveMetaData, InputCurveFrameValues, OutputCurveMetaData, OutputCurveFrameValues);
@@ -604,7 +604,7 @@ void UIKRetargetBatchOperation::ApplyCurveOps(const FIKRetargetBatchOperationCon
 			// set the output values and metadata into the anim sequence
 			// note that outputs are contained in the Input data variables since they are moved there after each operation
 			const FFrameRate InputFrameRate = FCurveUtils::GetAnimSequenceRate(SourceSequence);
-			AddCurveValuesToAnimSequence(SourceSkeleton, InputCurveMetaData, InputCurveFrameValues, InputFrameRate, InputFrameTimes, bShouldTransact, TargetSeqController);
+			AddCurveValuesToAnimSequence(TargetSkeleton, InputCurveMetaData, InputCurveFrameValues, InputFrameRate, InputFrameTimes, bShouldTransact, TargetSeqController);
 		}
 			
 		if(!bLoadedData)
@@ -617,7 +617,7 @@ void UIKRetargetBatchOperation::ApplyCurveOps(const FIKRetargetBatchOperationCon
 }
 
 
-void UIKRetargetBatchOperation::AddCurveValuesToAnimSequence(USkeleton* InTargetSkeleton, const FIKRetargetOpBase::FCurveData& InCurveMetaData, const FIKRetargetOpBase::FFrameValues& InCurveValuesPerFrame, 
+void UIKRetargetBatchOperation::AddCurveValuesToAnimSequence(USkeleton* InTargetSkeleton, const FIKRetargetCurvesOpBase::FCurveData& InCurveMetaData, const FIKRetargetCurvesOpBase::FFrameValues& InCurveValuesPerFrame,
 	const FFrameRate& InFrameRate, const TArray<FFrameTime> & InFrameTimes, bool bInShouldTransact, IAnimationDataController& OutTargetSeqController) const
 {
 	// remove existing float curves, as we will replace them all
diff --git a/Engine/Plugins/Animation/IKRig/Source/IKRigEditor/Private/RetargetEditor/SRetargetOpStack.cpp b/Engine/Plugins/Animation/IKRig/Source/IKRigEditor/Private/RetargetEditor/SRetargetOpStack.cpp
--- a/Engine/Plugins/Animation/IKRig/Source/IKRigEditor/Private/RetargetEditor/SRetargetOpStack.cpp
+++ b/Engine/Plugins/Animation/IKRig/Source/IKRigEditor/Private/RetargetEditor/SRetargetOpStack.cpp
@@ -377,6 +377,12 @@ void SRetargetOpList::CacheOpTypeMetaData()
 		{
 			continue;
 		}
+
+		// disallow adding abstract types to the stack
+		if (TempOp->IsAbstract())
+		{
+			continue;
+		}
 		
 		FIKRetargetOpMetaData OpMetaData;
 		OpMetaData.NiceName = TempOp->GetDefaultName();
diff --git a/Engine/Plugins/Animation/IKRig/Source/IKRigEditor/Public/RetargetEditor/IKRetargetBatchOperation.h b/Engine/Plugins/Animation/IKRig/Source/IKRigEditor/Public/RetargetEditor/IKRetargetBatchOperation.h
--- a/Engine/Plugins/Animation/IKRig/Source/IKRigEditor/Public/RetargetEditor/IKRetargetBatchOperation.h
+++ b/Engine/Plugins/Animation/IKRig/Source/IKRigEditor/Public/RetargetEditor/IKRetargetBatchOperation.h
@@ -2,9 +2,8 @@
 
 #pragma once
 
-#include "CoreMinimal.h"
 #include "EditorAnimUtils.h"
-#include "Retargeter/IKRetargeter.h"
+#include "Retargeter/RetargetOps/RetargetCurvesOp.h"
 
 #include "IKRetargetBatchOperation.generated.h"
 
@@ -185,8 +184,8 @@ private:
 	UE_API void CleanupIfCancelled(const FScopedSlowTask& Progress) const;
 
 	// Set the processed curve data into the target animation sequence
-	UE_API void AddCurveValuesToAnimSequence(USkeleton* InTargetSkeleton, const FIKRetargetOpBase::FCurveData& InCurveMetaData,
-		const FIKRetargetOpBase::FFrameValues& InCurveValuesPerFrame, const FFrameRate& InFrameRate, const TArray<FFrameTime>& InFrameTimes, bool bInShouldTransact, IAnimationDataController& OutTargetSeqController) const;
+	UE_API void AddCurveValuesToAnimSequence(USkeleton* InTargetSkeleton, const FIKRetargetCurvesOpBase::FCurveData& InCurveMetaData,
+		const FIKRetargetCurvesOpBase::FFrameValues& InCurveValuesPerFrame, const FFrameRate& InFrameRate, const TArray<FFrameTime>& InFrameTimes, bool bInShouldTransact, IAnimationDataController& OutTargetSeqController) const;
 
 	
 	// Lists of assets to retarget. Populated from selection during init
diff --git a/Engine/Plugins/Experimental/Animation/RigMapperOp/Source/RigMapperOp/Private/RigMapperOp.cpp b/Engine/Plugins/Experimental/Animation/RigMapperOp/Source/RigMapperOp/Private/RigMapperOp.cpp
--- a/Engine/Plugins/Experimental/Animation/RigMapperOp/Source/RigMapperOp/Private/RigMapperOp.cpp
+++ b/Engine/Plugins/Experimental/Animation/RigMapperOp/Source/RigMapperOp/Private/RigMapperOp.cpp
@@ -21,10 +21,10 @@ static FLinearColor MakeCurveColor(const FName& InCurveName)
 }
 
 void FRigMapperOpHelper::ProcessAnimSequenceCurves(
-	FIKRetargetOpBase::FCurveData InCurveMetaData, 
-	FIKRetargetOpBase::FFrameValues InCurveFrameValues, 
-	FIKRetargetOpBase::FCurveData& OutCurveMetaData, 
-	FIKRetargetOpBase::FFrameValues& OutCurveFrameValues,
+	FIKRetargetCurvesOpBase::FCurveData InCurveMetaData,
+	FIKRetargetCurvesOpBase::FFrameValues InCurveFrameValues,
+	FIKRetargetCurvesOpBase::FCurveData& OutCurveMetaData,
+	FIKRetargetCurvesOpBase::FFrameValues& OutCurveFrameValues,
 	bool bCopyAllSourceCurves) const
 {
 	// We need to copy the RigMapperProcessor to keep this method const. Revisit this for speed at some point.
@@ -79,6 +79,12 @@ void FRigMapperOpHelper::ProcessAnimSequenceCurves(
 			}
 		}
 	}
+	else
+	{
+		// Fallback when the RigMapper processor is not valid, tp avoid empty output curves.
+		OutCurveMetaData = MoveTemp(InCurveMetaData);
+		OutCurveFrameValues = MoveTemp(InCurveFrameValues);
+	}
 }
 
 bool FRigMapperOpHelper::CheckReInit(const TArray<URigMapperDefinition*>& InCurrentDefinitions)
@@ -194,7 +200,8 @@ bool FIKRetargetRigMapperOp::Initialize(const FIKRetargetProcessor& InProcessor,
 	return bIsInitialized;
 }
 
-void FIKRetargetRigMapperOp::ProcessAnimSequenceCurves(FIKRetargetOpBase::FCurveData InCurveMetaData, FIKRetargetOpBase::FFrameValues InCurveFrameValues, FIKRetargetOpBase::FCurveData& OutCurveMetaData, FIKRetargetOpBase::FFrameValues& OutCurveFrameValues) const
+void FIKRetargetRigMapperOp::ProcessAnimSequenceCurves(FIKRetargetCurvesOpBase::FCurveData InCurveMetaData, FIKRetargetCurvesOpBase::FFrameValues InCurveFrameValues, 
+	FIKRetargetCurvesOpBase::FCurveData& OutCurveMetaData, FIKRetargetCurvesOpBase::FFrameValues& OutCurveFrameValues) const
 {
 	Helper.ProcessAnimSequenceCurves(InCurveMetaData, InCurveFrameValues, OutCurveMetaData, OutCurveFrameValues, Settings.bCopyAllSourceCurves);
 }
@@ -360,7 +367,8 @@ bool FIKRetargetRigMapperUserDataOp::Initialize(const FIKRetargetProcessor& InPr
 	return bIsInitialized;
 }
 
-void FIKRetargetRigMapperUserDataOp::ProcessAnimSequenceCurves(FIKRetargetOpBase::FCurveData InCurveMetaData, FIKRetargetOpBase::FFrameValues InCurveFrameValues, FIKRetargetOpBase::FCurveData& OutCurveMetaData, FIKRetargetOpBase::FFrameValues& OutCurveFrameValues) const
+void FIKRetargetRigMapperUserDataOp::ProcessAnimSequenceCurves(FIKRetargetCurvesOpBase::FCurveData InCurveMetaData, FIKRetargetCurvesOpBase::FFrameValues InCurveFrameValues, 
+	FIKRetargetCurvesOpBase::FCurveData& OutCurveMetaData, FIKRetargetCurvesOpBase::FFrameValues& OutCurveFrameValues) const
 {
 	Helper.ProcessAnimSequenceCurves(InCurveMetaData, InCurveFrameValues, OutCurveMetaData, OutCurveFrameValues, Settings.bCopyAllSourceCurves);
 }
diff --git a/Engine/Plugins/Experimental/Animation/RigMapperOp/Source/RigMapperOp/Public/RigMapperOp.h b/Engine/Plugins/Experimental/Animation/RigMapperOp/Source/RigMapperOp/Public/RigMapperOp.h
--- a/Engine/Plugins/Experimental/Animation/RigMapperOp/Source/RigMapperOp/Public/RigMapperOp.h
+++ b/Engine/Plugins/Experimental/Animation/RigMapperOp/Source/RigMapperOp/Public/RigMapperOp.h
@@ -2,6 +2,7 @@
 #pragma once
 
 #include "Retargeter/IKRetargetOps.h"
+#include "Retargeter/RetargetOps/RetargetCurvesOp.h"
 #include "Retargeter/RetargetOps/CurveRemapOp.h"
 #include "Animation/AnimCurveTypes.h"
 #include "RigMapperProcessor.h"
@@ -71,8 +72,8 @@ public:
 
 	bool InitializeRigMapping(const TArray<URigMapperDefinition*>& InDefinitions);
 
-	void ProcessAnimSequenceCurves(FIKRetargetOpBase::FCurveData InCurveMetaData, FIKRetargetOpBase::FFrameValues InCurveFrameValues,
-		FIKRetargetOpBase::FCurveData& OutCurveMetaData, FIKRetargetOpBase::FFrameValues& OutCurveFrameValues, bool bCopyAllSourceCurves) const;
+	void ProcessAnimSequenceCurves(FIKRetargetCurvesOpBase::FCurveData InCurveMetaData, FIKRetargetCurvesOpBase::FFrameValues InCurveFrameValues,
+		FIKRetargetCurvesOpBase::FCurveData& OutCurveMetaData, FIKRetargetCurvesOpBase::FFrameValues& OutCurveFrameValues, bool bCopyAllSourceCurves) const;
 
 	void EvaluateRigMapping(const FBlendedCurve& InCurve, FBlendedCurve& OutCurve);
 
@@ -136,7 +137,7 @@ private:
 
 // Single definition RigMapper op with an option to use UserData definitions array
 USTRUCT(BlueprintType, meta = (DisplayName = "Single RigMapper"))
-struct FIKRetargetRigMapperOp : public FIKRetargetOpBase
+struct FIKRetargetRigMapperOp : public FIKRetargetCurvesOpBase
 {
 	GENERATED_BODY()
 	
@@ -180,10 +181,12 @@ public:
 		return FIKRetargetRigMapperOp::StaticStruct();
 	}
 
+	virtual bool IsAbstract() const override { return false; };
+
 	virtual bool IsSingleton() const override { return false; }
 
-	virtual void ProcessAnimSequenceCurves(FIKRetargetOpBase::FCurveData InCurveMetaData, FIKRetargetOpBase::FFrameValues InCurveFrameValues,
-		FIKRetargetOpBase::FCurveData& OutCurveMetaData, FIKRetargetOpBase::FFrameValues& OutCurveFrameValues) const;
+	virtual void ProcessAnimSequenceCurves(FIKRetargetCurvesOpBase::FCurveData InCurveMetaData, FIKRetargetCurvesOpBase::FFrameValues InCurveFrameValues,
+		FIKRetargetCurvesOpBase::FCurveData& OutCurveMetaData, FIKRetargetCurvesOpBase::FFrameValues& OutCurveFrameValues) const override;
 
 	UE_API virtual void AnimGraphPreUpdateMainThread(USkeletalMeshComponent& SourceMeshComponent, USkeletalMeshComponent& TargetMeshComponent) override;
 
@@ -194,9 +197,6 @@ public:
 	UPROPERTY()
 	FIKRetargetRigMapperOpSettings Settings;
 
-	// Cached curves, copied on the game thread in PreUpdate()
-	FBlendedHeapCurve SourceCurves;
-
 protected:
 	TArray<URigMapperDefinition*> GetDefinitionsToLoad(const USkeletalMesh* InTargetMesh);
 
@@ -224,7 +224,7 @@ public:
 
 // UserData definitions ONLY RigMapper op
 USTRUCT(BlueprintType, meta = (DisplayName = "UserData RigMapper"))
-struct FIKRetargetRigMapperUserDataOp : public FIKRetargetOpBase
+struct FIKRetargetRigMapperUserDataOp : public FIKRetargetCurvesOpBase
 {
 	GENERATED_BODY()
 
@@ -269,18 +269,17 @@ public:
 		return FIKRetargetRigMapperUserDataOp::StaticStruct();
 	}
 
+	virtual bool IsAbstract() const override { return false; };
+
 	virtual bool IsSingleton() const override { return false; }
 
-	virtual void ProcessAnimSequenceCurves(FIKRetargetOpBase::FCurveData InCurveMetaData, FIKRetargetOpBase::FFrameValues InCurveFrameValues,
-		FIKRetargetOpBase::FCurveData& OutCurveMetaData, FIKRetargetOpBase::FFrameValues& OutCurveFrameValues) const;
+	virtual void ProcessAnimSequenceCurves(FIKRetargetCurvesOpBase::FCurveData InCurveMetaData, FIKRetargetCurvesOpBase::FFrameValues InCurveFrameValues,
+		FIKRetargetCurvesOpBase::FCurveData& OutCurveMetaData, FIKRetargetCurvesOpBase::FFrameValues& OutCurveFrameValues) const;
 
 	UE_API virtual void AnimGraphPreUpdateMainThread(USkeletalMeshComponent& SourceMeshComponent, USkeletalMeshComponent& TargetMeshComponent) override;
 
 	UE_API virtual void AnimGraphEvaluateAnyThread(FPoseContext& Output) override;
 
-	// Cached curves, copied on the game thread in PreUpdate()
-	FBlendedHeapCurve SourceCurves;
-
 	UPROPERTY()
 	FIKRetargetRigMapperUserDataOpSettings Settings;
 
